MRMesh circular dependency
-Removed dependency to HMD.

diff --git a/Engine/Source/Runtime/MRMesh/MRMesh.build.cs b/Engine/Source/Runtime/MRMesh/MRMesh.build.cs
--- a/Engine/Source/Runtime/MRMesh/MRMesh.build.cs
+++ b/Engine/Source/Runtime/MRMesh/MRMesh.build.cs
@@ -16,8 +16,7 @@ namespace UnrealBuildTool.Rules
 					"CoreUObject",
                     "Engine",
                     "RenderCore",
-                    "RHI",
-					"HeadMountedDisplay"
+                    "RHI"
 				}
 			);
 
diff --git a/Engine/Source/Runtime/MRMesh/Private/MockDataMeshTrackerComponent.cpp b/Engine/Source/Runtime/MRMesh/Private/MockDataMeshTrackerComponent.cpp
--- a/Engine/Source/Runtime/MRMesh/Private/MockDataMeshTrackerComponent.cpp
+++ b/Engine/Source/Runtime/MRMesh/Private/MockDataMeshTrackerComponent.cpp
@@ -4,7 +4,6 @@
 #include "Engine/Engine.h"
 #include "Engine/GameEngine.h"
 #include "GameFramework/WorldSettings.h"
-#include "HeadMountedDisplayFunctionLibrary.h"
 #if WITH_EDITOR
 #include "Editor.h"
 #endif
@@ -439,7 +438,7 @@ void UMockDataMeshTrackerComponent::UpdateBlock(int32 BlockIndex)
 	CurrentMeshDataCache->BrickId = BrickId;
 
 	// Pull vertices
-	const FVector VertexOffset = UHeadMountedDisplayFunctionLibrary::GetTrackingToWorldTransform(this).Inverse().GetLocation();
+	const FVector VertexOffset = FVector::ZeroVector; // If this was the inverse of the worldToTrackingTransform position the mesh would be located in tracker space, but we don't want a dependency on HMD here.  
 	const int32 VertexCount = RawMeshData.Vertices.Num();
 	CurrentMeshDataCache->OffsetVertices.Reserve(VertexCount);
 	CurrentMeshDataCache->WorldVertices.Reserve(VertexCount);
