Duplicating for 4.21 Change: 4446057 UE-64942 Material Switching in MRMesh
-Fixed run time material switching on MRMesh

Files:
//UE4/Dev-VR/Engine/Source/Runtime/MRMesh/Private/MRMeshComponent.cpp#24
//UE4/Dev-VR/Engine/Source/Runtime/MRMesh/Public/MRMeshComponent.h#16

diff --git a/Engine/Source/Runtime/MRMesh/Private/MRMeshComponent.cpp b/Engine/Source/Runtime/MRMesh/Private/MRMeshComponent.cpp
--- a/Engine/Source/Runtime/MRMesh/Private/MRMeshComponent.cpp
+++ b/Engine/Source/Runtime/MRMesh/Private/MRMeshComponent.cpp
@@ -126,7 +126,7 @@ struct FMRMeshProxySection
 
 static void InitVertexFactory(FLocalVertexFactory* VertexFactory, const FMRMeshProxySection& MRMeshSection)
 {
-	ENQUEUE_RENDER_COMMAND(InitProcMeshVertexFactory)(
+	ENQUEUE_RENDER_COMMAND(InitMrMeshVertexFactory)(
 		[VertexFactory, &MRMeshSection](FRHICommandListImmediate& RHICmdList)
 	{
 		check(IsInRenderingThread());
@@ -286,6 +286,11 @@ public:
 		}
 	}
 
+	void RenderThread_SetMaterial(UMaterialInterface* Material)
+	{
+		MaterialToUse = Material;
+	}
+
 private:
 	//~ FPrimitiveSceneProxy
 
@@ -577,7 +582,28 @@ void UMRMeshComponent::SetMaterial(int32 ElementIndex, class UMaterialInterface*
 	if (Material != InMaterial)
 	{
 		Material = InMaterial;
-		MarkRenderStateDirty();
+		MarkRenderDynamicDataDirty();
+	}
+}
+
+void UMRMeshComponent::SendRenderDynamicData_Concurrent()
+{
+	Super::SendRenderDynamicData_Concurrent();
+
+	if (SceneProxy)
+	{
+		// Enqueue command to send to render thread
+		ENQUEUE_UNIQUE_RENDER_COMMAND_TWOPARAMETER(
+			FSetMaterialLambda,
+			UMRMeshComponent*, This, this,
+			UMaterialInterface*, Material, Material,
+			{
+				FMRMeshProxy* MRMeshProxy = static_cast<FMRMeshProxy*>(This->SceneProxy);
+				if (MRMeshProxy)
+				{
+					MRMeshProxy->RenderThread_SetMaterial(Material);
+				}
+			});
 	}
 }
 
diff --git a/Engine/Source/Runtime/MRMesh/Public/MRMeshComponent.h b/Engine/Source/Runtime/MRMesh/Public/MRMeshComponent.h
--- a/Engine/Source/Runtime/MRMesh/Public/MRMeshComponent.h
+++ b/Engine/Source/Runtime/MRMesh/Public/MRMeshComponent.h
@@ -79,6 +79,9 @@ private:
 	virtual void SetMaterial(int32 ElementIndex, class UMaterialInterface* InMaterial) override;
 	virtual bool DoCustomNavigableGeometryExport(FNavigableGeometryExport& GeomExport) const override;
 	//~ UPrimitiveComponent
+	//~ UActorComponent
+	virtual void SendRenderDynamicData_Concurrent() override;
+	//~ UActorComponent
 
 	//~ IMRMesh
 	virtual void SendBrickData(FSendBrickDataArgs Args) override;
