UE-59359 Mr Mesh Crash at startup
-Filling in dummy tangent data that lower level vulkan code assumes is not null.
-Also fixed the vr.MagicLeap.FakeMeshTrackerData cvar, it was missing the 'MagicLeap'.

diff --git a/Engine/Plugins/Lumin/MagicLeap/Source/MagicLeap/Private/MeshTrackerComponent.cpp b/Engine/Plugins/Lumin/MagicLeap/Source/MagicLeap/Private/MeshTrackerComponent.cpp
--- a/Engine/Plugins/Lumin/MagicLeap/Source/MagicLeap/Private/MeshTrackerComponent.cpp
+++ b/Engine/Plugins/Lumin/MagicLeap/Source/MagicLeap/Private/MeshTrackerComponent.cpp
@@ -519,7 +519,7 @@ void UMeshTrackerComponent::TickComponent(float DeltaTime, enum ELevelTick TickT
 		return;
 	}
 
-	static const auto FakeMeshTrackerDataCVar = IConsoleManager::Get().FindConsoleVariable(TEXT("vr.FakeMeshTrackerData"));
+	static const auto FakeMeshTrackerDataCVar = IConsoleManager::Get().FindConsoleVariable(TEXT("vr.MagicLeap.FakeMeshTrackerData"));
 	if (FakeMeshTrackerDataCVar != nullptr && FakeMeshTrackerDataCVar->GetInt())
 	{
 		TickWithFakeData();
diff --git a/Engine/Source/Runtime/MRMesh/Private/MRMeshComponent.cpp b/Engine/Source/Runtime/MRMesh/Private/MRMeshComponent.cpp
--- a/Engine/Source/Runtime/MRMesh/Private/MRMeshComponent.cpp
+++ b/Engine/Source/Runtime/MRMesh/Private/MRMeshComponent.cpp
@@ -138,12 +138,13 @@ static void InitVertexFactory(FLocalVertexFactory* VertexFactory, const FMRMeshP
 			NewData.PositionComponentSRV = MRMeshSection.PositionBufferSRV;
 			NewData.PositionComponent = FVertexStreamComponent(&MRMeshSection.PositionBuffer, 0, sizeof(FVector), VET_Float3, EVertexStreamUsage::Default);
 		}
-		
-		//{
-		//	NewData.TextureCoordinatesSRV = MRMeshSection.UVBufferSRV;
-		//	NewData.TextureCoordinates.Add(FVertexStreamComponent(&MRMeshSection.UVBuffer, 0, sizeof(FVector2D), VET_Float2, EVertexStreamUsage::ManualFetch));
-		//	NewData.NumTexCoords = 1;
-		//}
+
+		if (MRMeshSection.UVBuffer.NumVerts != 0)
+		{
+			NewData.TextureCoordinatesSRV = MRMeshSection.UVBufferSRV;
+			NewData.TextureCoordinates.Add(FVertexStreamComponent(&MRMeshSection.UVBuffer, 0, sizeof(FVector2D), VET_Float2, EVertexStreamUsage::ManualFetch));
+			NewData.NumTexCoords = 1;
+		}
 
 		if (MRMeshSection.TangentXZBuffer.NumVerts != 0)
 		{
@@ -199,9 +200,10 @@ public:
 		FMRMeshProxySection* NewSection = new FMRMeshProxySection(Args.BrickId, FeatureLevel);
 		ProxySections.Add(NewSection);
 
-		check((Args.PositionData.Num() == Args.ColorData.Num()) || Args.ColorData.Num() == 0);
-		check((Args.PositionData.Num() == Args.UVData.Num()) || Args.UVData.Num() == 0);
-		check(((Args.PositionData.Num() * 2) == Args.TangentXZData.Num()) || Args.TangentXZData.Num() == 0);
+		const int32 NumVerts = Args.PositionData.Num();
+		check((NumVerts == Args.ColorData.Num()) || Args.ColorData.Num() == 0);
+		check((NumVerts == Args.UVData.Num()) || Args.UVData.Num() == 0);
+		check(((NumVerts * 2) == Args.TangentXZData.Num()) || Args.TangentXZData.Num() == 0);
 
 		// POSITION BUFFER
 		{
@@ -226,10 +228,23 @@ public:
 			if (Args.TangentXZData.Num())
 			{
 				NewSection->TangentXZBuffer.InitRHIWith(Args.TangentXZData);
-				if (RHISupportsManualVertexFetch(GMaxRHIShaderPlatform))
+			}
+			else
+			{
+				// Vulkan requires that the tangent data exist, so we fill it in with fake data if necessary.
+				TArray<FPackedNormal> FakeData;
+				FakeData.Reserve(NumVerts * 2);
+				for (int32 i = 0; i < NumVerts; i++)
 				{
-					NewSection->TangentXZBufferSRV = RHICreateShaderResourceView(NewSection->TangentXZBuffer.VertexBufferRHI, 4, PF_R8G8B8A8_SNORM);
+					FakeData.Add(FPackedNormal(FVector::UpVector));
+					FakeData.Add(FPackedNormal(FVector::RightVector));
 				}
+				NewSection->TangentXZBuffer.InitRHIWith(FakeData);
+			}
+
+			if (RHISupportsManualVertexFetch(GMaxRHIShaderPlatform))
+			{
+				NewSection->TangentXZBufferSRV = RHICreateShaderResourceView(NewSection->TangentXZBuffer.VertexBufferRHI, 4, PF_R8G8B8A8_SNORM);
 			}
 		}
 
