UE-64940 Collision clearing on MRmesh
-Added support for generic runtime collision modification of MRMesh collision.
-Generic blueprint collision settings now work on MRMesh.

diff --git a/Engine/Source/Runtime/Engine/Classes/Components/PrimitiveComponent.h b/Engine/Source/Runtime/Engine/Classes/Components/PrimitiveComponent.h
--- a/Engine/Source/Runtime/Engine/Classes/Components/PrimitiveComponent.h
+++ b/Engine/Source/Runtime/Engine/Classes/Components/PrimitiveComponent.h
@@ -987,7 +987,7 @@ public:
 
 	/** Sets a new slope override for this component instance. */
 	UFUNCTION(BlueprintCallable, Category="Physics")
-	void SetWalkableSlopeOverride(const FWalkableSlopeOverride& NewOverride);
+	virtual void SetWalkableSlopeOverride(const FWalkableSlopeOverride& NewOverride);
 
 	/** 
 	 *	Sets whether or not a single body should use physics simulation, or should be 'fixed' (kinematic).
diff --git a/Engine/Source/Runtime/Engine/Classes/PhysicsEngine/BodyInstance.h b/Engine/Source/Runtime/Engine/Classes/PhysicsEngine/BodyInstance.h
--- a/Engine/Source/Runtime/Engine/Classes/PhysicsEngine/BodyInstance.h
+++ b/Engine/Source/Runtime/Engine/Classes/PhysicsEngine/BodyInstance.h
@@ -594,6 +594,9 @@ public:
 	/** Utility for copying properties from one BodyInstance to another. */
 	void CopyBodyInstancePropertiesFrom(const FBodyInstance* FromInst);
 
+	/** Utility for copying only the runtime instanced properties from one BodyInstance to another. */
+	void CopyRuntimeBodyInstancePropertiesFrom(const FBodyInstance* FromInst);
+
 	/** Find the correct PhysicalMaterial for simple geometry on this body */
 	UPhysicalMaterial* GetSimplePhysicalMaterial() const;
 
diff --git a/Engine/Source/Runtime/Engine/Private/PhysicsEngine/BodyInstance.cpp b/Engine/Source/Runtime/Engine/Private/PhysicsEngine/BodyInstance.cpp
--- a/Engine/Source/Runtime/Engine/Private/PhysicsEngine/BodyInstance.cpp
+++ b/Engine/Source/Runtime/Engine/Private/PhysicsEngine/BodyInstance.cpp
@@ -2252,10 +2252,27 @@ void FBodyInstance::CopyBodyInstancePropertiesFrom(const FBodyInstance* FromInst
 	*this = *FromInst;
 }
 
+void FBodyInstance::CopyRuntimeBodyInstancePropertiesFrom(const FBodyInstance* FromInst)
+{
+	check(FromInst);
+	check(!FromInst->bPendingCollisionProfileSetup);
+
+	if (FromInst->bOverrideWalkableSlopeOnInstance)
+	{
+		SetWalkableSlopeOverride(FromInst->GetWalkableSlopeOverride());
+	}
+
+	CollisionResponses = FromInst->CollisionResponses;
+	CollisionProfileName = FromInst->CollisionProfileName;
+	CollisionEnabled = FromInst->CollisionEnabled;
+
+	UpdatePhysicsFilterData();
+}
+
 FPhysScene* FBodyInstance::GetPhysicsScene() const
 {
 	if(ActorHandle.IsValid())
-{
+	{
 		return FPhysicsInterface::GetCurrentScene(ActorHandle);
 	}
 	
@@ -2518,86 +2535,90 @@ void FBodyInstance::UpdateMassProperties()
 
 				TArray<FPhysicsShapeHandle> Shapes;
 				const uint32 NumShapes = bSyncData ? NumSyncShapes : NumAsyncShapes;
-			Shapes.AddUninitialized(NumShapes);
+				Shapes.AddUninitialized(NumShapes);
 				FPhysicsInterface::GetAllShapes_AssumedLocked(Actor, Shapes, bSyncData ? PST_Sync : PST_Async);
 
-			// Ignore trimeshes & shapes which don't contribute to the mass
-			for(int32 ShapeIdx = Shapes.Num() - 1; ShapeIdx >= 0; --ShapeIdx)
-			{
+				// Ignore trimeshes & shapes which don't contribute to the mass
+				for(int32 ShapeIdx = Shapes.Num() - 1; ShapeIdx >= 0; --ShapeIdx)
+				{
 					const FPhysicsShapeHandle& Shape = Shapes[ShapeIdx];
 					const FKShapeElem* ShapeElem = FPhysxUserData::Get<FKShapeElem>(FPhysicsInterface::GetUserData(Shape));
 					bool bIsTriangleMesh = FPhysicsInterface::GetShapeType(Shape) == ECollisionShapeType::Trimesh;
-				bool bHasNoMass = ShapeElem && !ShapeElem->GetContributeToMass();
-				if (bIsTriangleMesh || bHasNoMass)
-				{
-					Shapes.RemoveAtSwap(ShapeIdx);
+					bool bHasNoMass = ShapeElem && !ShapeElem->GetContributeToMass();
+					if (bIsTriangleMesh || bHasNoMass)
+					{
+						Shapes.RemoveAtSwap(ShapeIdx);
+					}
 				}
-			}
 
-			PxMassProperties TotalMassProperties;
-			if(ShapeToBodiesMap.IsValid() && ShapeToBodiesMap->Num() > 0)
-			{
-				struct FWeldedBatch
+				PxMassProperties TotalMassProperties;
+				if(ShapeToBodiesMap.IsValid() && ShapeToBodiesMap->Num() > 0)
 				{
-						TArray<FPhysicsShapeHandle> Shapes;
-					FTransform RelTM;
-				};
+					struct FWeldedBatch
+					{
+							TArray<FPhysicsShapeHandle> Shapes;
+						FTransform RelTM;
+					};
 
-				//If we have welded children we must compute the mass properties of each individual body first and then combine them all together
-				TMap<FBodyInstance*, FWeldedBatch> BodyToShapes;
+					//If we have welded children we must compute the mass properties of each individual body first and then combine them all together
+					TMap<FBodyInstance*, FWeldedBatch> BodyToShapes;
 
 					for(const FPhysicsShapeHandle& Shape : Shapes) //sort all welded children by their original bodies
-				{
-					if (FWeldInfo* WeldInfo = ShapeToBodiesMap->Find(Shape))
 					{
-						FWeldedBatch* WeldedBatch = BodyToShapes.Find(WeldInfo->ChildBI);
-						if(!WeldedBatch)
+						if (FWeldInfo* WeldInfo = ShapeToBodiesMap->Find(Shape))
 						{
-							WeldedBatch = &BodyToShapes.Add(WeldInfo->ChildBI);
-							WeldedBatch->RelTM = WeldInfo->RelativeTM;
-						}
+							FWeldedBatch* WeldedBatch = BodyToShapes.Find(WeldInfo->ChildBI);
+							if(!WeldedBatch)
+							{
+								WeldedBatch = &BodyToShapes.Add(WeldInfo->ChildBI);
+								WeldedBatch->RelTM = WeldInfo->RelativeTM;
+							}
 
-						WeldedBatch->Shapes.Add(Shape);
-					}
-					else
-					{
-						//no weld info so shape really belongs to this body
-						FWeldedBatch* WeldedBatch = BodyToShapes.Find(this);
-						if (!WeldedBatch)
+							WeldedBatch->Shapes.Add(Shape);
+						}
+						else
 						{
-							WeldedBatch = &BodyToShapes.Add(this);
-							WeldedBatch->RelTM = FTransform::Identity;
+							//no weld info so shape really belongs to this body
+							FWeldedBatch* WeldedBatch = BodyToShapes.Find(this);
+							if (!WeldedBatch)
+							{
+								WeldedBatch = &BodyToShapes.Add(this);
+								WeldedBatch->RelTM = FTransform::Identity;
+							}
+
+							WeldedBatch->Shapes.Add(Shape);
 						}
+					}
 
-						WeldedBatch->Shapes.Add(Shape);
+					TArray<PxMassProperties> SubMassProperties;
+					TArray<PxTransform> MassTMs;
+					for(auto BodyShapesItr : BodyToShapes)
+					{
+						const FBodyInstance* OwningBI = BodyShapesItr.Key;
+						const FWeldedBatch& WeldedBatch = BodyShapesItr.Value;
+						FTransform MassModifierTransform = WeldedBatch.RelTM;
+						MassModifierTransform.SetScale3D(MassModifierTransform.GetScale3D() * Scale3D);	//Ensure that any scaling that is done on the component is passed into the mass frame modifiers
+
+						PxMassProperties BodyMassProperties = ComputeMassProperties(OwningBI, WeldedBatch.Shapes, MassModifierTransform);
+						SubMassProperties.Add(BodyMassProperties);
+						MassTMs.Add(PxTransform(PxIdentity));
 					}
-				}
 
-				TArray<PxMassProperties> SubMassProperties;
-				TArray<PxTransform> MassTMs;
-				for(auto BodyShapesItr : BodyToShapes)
+					TotalMassProperties = PxMassProperties::sum(SubMassProperties.GetData(), MassTMs.GetData(), SubMassProperties.Num());
+				}
+				else
 				{
-					const FBodyInstance* OwningBI = BodyShapesItr.Key;
-					const FWeldedBatch& WeldedBatch = BodyShapesItr.Value;
-					FTransform MassModifierTransform = WeldedBatch.RelTM;
-					MassModifierTransform.SetScale3D(MassModifierTransform.GetScale3D() * Scale3D);	//Ensure that any scaling that is done on the component is passed into the mass frame modifiers
-
-					PxMassProperties BodyMassProperties = ComputeMassProperties(OwningBI, WeldedBatch.Shapes, MassModifierTransform);
-					SubMassProperties.Add(BodyMassProperties);
-					MassTMs.Add(PxTransform(PxIdentity));
+					// If we have no shapes that affect mass we cannot compute the mass properties in a meaningful way.
+					if (Shapes.Num())
+					{
+						//No children welded so just get this body's mass properties
+						FTransform MassModifierTransform(FQuat::Identity, FVector(0.f, 0.f, 0.f), Scale3D);	//Ensure that any scaling that is done on the component is passed into the mass frame modifiers
+						TotalMassProperties = ComputeMassProperties(this, Shapes, MassModifierTransform);
+					}
 				}
-
-				TotalMassProperties = PxMassProperties::sum(SubMassProperties.GetData(), MassTMs.GetData(), SubMassProperties.Num());
-			}
-			else
-			{
-				//No children welded so just get this body's mass properties
-				FTransform MassModifierTransform(FQuat::Identity, FVector(0.f, 0.f, 0.f), Scale3D);	//Ensure that any scaling that is done on the component is passed into the mass frame modifiers
-				TotalMassProperties = ComputeMassProperties(this, Shapes, MassModifierTransform);
-			}
 			
 				// #PHYS2 Refactor out PxMassProperties (Our own impl?)
-			PxQuat MassOrientation;
+				PxQuat MassOrientation;
 				const FVector MassSpaceInertiaTensor = P2UVector(PxMassProperties::getMassSpaceInertia(TotalMassProperties.inertiaTensor, MassOrientation));
 
 				FPhysicsInterface::SetMass_AssumesLocked(Actor, TotalMassProperties.mass);
@@ -2605,8 +2626,8 @@ void FBodyInstance::UpdateMassProperties()
 
 				FTransform Com(P2UQuat(MassOrientation), P2UVector(TotalMassProperties.centerOfMass));
 				FPhysicsInterface::SetComLocalPose_AssumesLocked(Actor, Com);
-		}
-	});
+			}
+		});
 	}
 #endif
 
diff --git a/Engine/Source/Runtime/MRMesh/Private/MRMeshComponent.cpp b/Engine/Source/Runtime/MRMesh/Private/MRMeshComponent.cpp
--- a/Engine/Source/Runtime/MRMesh/Private/MRMeshComponent.cpp
+++ b/Engine/Source/Runtime/MRMesh/Private/MRMeshComponent.cpp
@@ -20,6 +20,10 @@
 #include "PhysXCookHelper.h"
 #include "Misc/RuntimeErrors.h"
 #include "Engine/Engine.h"
+#include "UObject/UObjectThreadContext.h"
+#include "Stats/Stats.h"
+
+DECLARE_CYCLE_STAT(TEXT("MrMesh SetCollisionProfileName"), STAT_MrMesh_SetCollisionProfileName, STATGROUP_Physics);
 
 class FMRMeshVertexResourceArray : public FResourceArrayInterface
 {
@@ -376,6 +380,140 @@ void UMRMeshComponent::BeginPlay()
 	SetCustomNavigableGeometry(bCanEverAffectNavigation ? EHasCustomNavigableGeometry::Yes : EHasCustomNavigableGeometry::No);
 }
 
+void UMRMeshComponent::EndPlay(const EEndPlayReason::Type EndPlayReason)
+{
+	ClearAllBrickData();
+
+	Super::EndPlay(EndPlayReason);
+}
+
+void UMRMeshComponent::OnActorEnableCollisionChanged()
+{
+	for (auto& BodyInstanceElement : BodyInstances)
+	{
+		BodyInstanceElement->UpdatePhysicsFilterData();
+	}
+
+	Super::OnActorEnableCollisionChanged();
+}
+
+bool UMRMeshComponent::ShouldCreatePhysicsState() const
+{
+	// This component does not use the default physics state creation.  It creates in response to meshing data delivered via SendBrickData.
+	return false;
+}
+
+void UMRMeshComponent::SetCollisionEnabled(ECollisionEnabled::Type NewType)
+{
+	if (BodyInstance.GetCollisionEnabled() != NewType)
+	{
+		for (auto& BodyInstanceElement : BodyInstances)
+		{
+			BodyInstanceElement->SetCollisionEnabled(NewType);
+		}
+
+		if (IsRegistered() && BodyInstance.bSimulatePhysics && !IsWelded())
+		{
+			for (auto& BodyInstanceElement : BodyInstances)
+			{
+				BodyInstanceElement->ApplyWeldOnChildren();
+			}
+		}
+	}
+
+	Super::SetCollisionEnabled(NewType);
+}
+
+void UMRMeshComponent::SetCollisionProfileName(FName InCollisionProfileName)
+{
+	SCOPE_CYCLE_COUNTER(STAT_MrMesh_SetCollisionProfileName);
+
+	FUObjectThreadContext& ThreadContext = FUObjectThreadContext::Get();
+	if (ThreadContext.ConstructedObject == this)
+	{
+		// If we are in our constructor, defer setup until PostInitProperties as derived classes
+		for (auto& BodyInstanceElement : BodyInstances)
+		{
+			BodyInstanceElement->SetCollisionProfileNameDeferred(InCollisionProfileName);
+		}
+	}
+	else
+	{
+		for (auto& BodyInstanceElement : BodyInstances)
+		{
+			BodyInstanceElement->SetCollisionProfileName(InCollisionProfileName);
+		}
+	}
+
+	Super::SetCollisionProfileName(InCollisionProfileName);
+}
+
+void UMRMeshComponent::SetCollisionObjectType(ECollisionChannel Channel)
+{
+	for (auto& BodyInstanceElement : BodyInstances)
+	{
+		BodyInstanceElement->SetObjectType(Channel);
+	}
+
+	Super::SetCollisionObjectType(Channel);
+}
+
+void UMRMeshComponent::SetCollisionResponseToChannel(ECollisionChannel Channel, ECollisionResponse NewResponse)
+{
+	for (auto& BodyInstanceElement : BodyInstances)
+	{
+		BodyInstanceElement->SetResponseToChannel(Channel, NewResponse);
+	}
+
+	Super::SetCollisionResponseToChannel(Channel, NewResponse);
+}
+
+void UMRMeshComponent::SetCollisionResponseToAllChannels(enum ECollisionResponse NewResponse)
+{
+	for (auto& BodyInstanceElement : BodyInstances)
+	{
+		BodyInstanceElement->SetResponseToAllChannels(NewResponse);
+	}
+
+	Super::SetCollisionResponseToAllChannels(NewResponse);
+}
+
+void UMRMeshComponent::SetCollisionResponseToChannels(const FCollisionResponseContainer& NewResponses)
+{
+	for (auto& BodyInstanceElement : BodyInstances)
+	{
+		BodyInstanceElement->SetResponseToChannels(NewResponses);
+	}
+
+	Super::SetCollisionResponseToChannels(NewResponses);
+}
+
+void UMRMeshComponent::UpdatePhysicsToRBChannels()
+{
+	for (auto& BodyInstanceElement : BodyInstances)
+	{
+		if (BodyInstanceElement->IsValidBodyInstance())
+		{
+			BodyInstanceElement->UpdatePhysicsFilterData();
+		}
+	}
+
+	Super::UpdatePhysicsToRBChannels();
+}
+
+void UMRMeshComponent::SetWalkableSlopeOverride(const FWalkableSlopeOverride& NewOverride)
+{
+	for (auto& BodyInstanceElement : BodyInstances)
+	{
+		if (BodyInstanceElement->IsValidBodyInstance())
+		{
+			BodyInstanceElement->SetWalkableSlopeOverride(NewOverride);
+		}
+	}
+
+	Super::SetWalkableSlopeOverride(NewOverride);
+}
+
 FPrimitiveSceneProxy* UMRMeshComponent::CreateSceneProxy()
 {
 	// The render thread owns the memory, so if this function is
@@ -418,16 +556,41 @@ void UMRMeshComponent::ClearAllBrickData()
 	FSimpleDelegateGraphTask::CreateAndDispatchWhenReady(ClearBrickDataTask, GET_STATID(STAT_UMRMeshComponent_ClearAllBrickData), nullptr, ENamedThreads::GameThread);
 }
 
-UBodySetup* CreateBodySetupHelper(UMRMeshComponent* Outer)
+void UMRMeshComponent::CacheBodySetupHelper()
+{
+	CachedBodySetup = NewObject<UBodySetup>(this, NAME_None);
+	CachedBodySetup->BodySetupGuid = FGuid::NewGuid();
+	CachedBodySetup->bGenerateMirroredCollision = false;
+	CachedBodySetup->bHasCookedCollisionData = true;
+}
+
+UBodySetup* UMRMeshComponent::CreateBodySetupHelper()
 {
 	// The body setup in a template needs to be public since the property is Instanced and thus is the archetype of the instance meaning there is a direct reference
-	UBodySetup* NewBS = NewObject<UBodySetup>(Outer, NAME_None);
+	UBodySetup* NewBS = NewObject<UBodySetup>(this, NAME_None);
 	NewBS->BodySetupGuid = FGuid::NewGuid();
 	NewBS->bGenerateMirroredCollision = false;
 	NewBS->bHasCookedCollisionData = true;
+
+	// Copy the cached body setup (unless we are creating it)
+	if (CachedBodySetup == nullptr)
+	{
+		CacheBodySetupHelper();
+	}
+	NewBS->CopyBodyPropertiesFrom(CachedBodySetup);
+
 	return NewBS;
 }
 
+UBodySetup* UMRMeshComponent::GetBodySetup()
+{
+	if (CachedBodySetup == nullptr)
+	{
+		CacheBodySetupHelper();
+	}
+	return CachedBodySetup;
+}
+
 void UMRMeshComponent::SendBrickData_Internal(IMRMesh::FSendBrickDataArgs Args)
 {
 #if WITH_PHYSX
@@ -435,7 +598,7 @@ void UMRMeshComponent::SendBrickData_Internal(IMRMesh::FSendBrickDataArgs Args)
 
 	UE_LOG(LogMrMesh, Log, TEXT("SendBrickData_Internal() processing brick %llu with %i triangles"), Args.BrickId, Args.Indices.Num() / 3);
 
-	if (!IsPendingKill() && this->GetCollisionEnabled())
+	if (!IsPendingKill() && !bNeverCreateCollisionMesh)
 	{
 		// Physics update
 		UWorld* MyWorld = GetWorld();
@@ -445,10 +608,12 @@ void UMRMeshComponent::SendBrickData_Internal(IMRMesh::FSendBrickDataArgs Args)
 
 			if (const bool bBrickHasData = Args.Indices.Num() > 0)
 			{
+				bPhysicsStateCreated = true;
+
 				if (BodyIndex == INDEX_NONE)
 				{
 					BodyIds.Add(Args.BrickId);
-					BodySetups.Add(CreateBodySetupHelper(this));
+					BodySetups.Add(CreateBodySetupHelper());
 					BodyInstances.Add(new FBodyInstance());
 					BodyIndex = BodyIds.Num() - 1;
 				}
@@ -456,6 +621,8 @@ void UMRMeshComponent::SendBrickData_Internal(IMRMesh::FSendBrickDataArgs Args)
 				UBodySetup* MyBS = BodySetups[BodyIndex];
 				MyBS->bHasCookedCollisionData = true;
 				MyBS->CollisionTraceFlag = CTF_UseComplexAsSimple;
+				MyBS->ClearPhysicsMeshes();
+				MyBS->InvalidatePhysicsData();
 
 				FCookBodySetupInfo CookInfo;
 				// Disable mesh cleaning by passing in EPhysXMeshCookFlags::DeformableMesh
@@ -479,12 +646,12 @@ void UMRMeshComponent::SendBrickData_Internal(IMRMesh::FSendBrickDataArgs Args)
 				CookHelper.CookInfo = CookInfo;
 				CookHelper.CreatePhysicsMeshes_Concurrent();
 
-				MyBS->InvalidatePhysicsData();
 				MyBS->FinishCreatingPhysicsMeshes(CookHelper.OutNonMirroredConvexMeshes, CookHelper.OutMirroredConvexMeshes, CookHelper.OutTriangleMeshes);
 
 				FBodyInstance* MyBI = BodyInstances[BodyIndex];
 				MyBI->TermBody();
 				MyBI->InitBody(MyBS, GetComponentTransform(), this, MyWorld->GetPhysicsScene());
+				MyBI->CopyRuntimeBodyInstancePropertiesFrom(&BodyInstance);
 			}
 			else
 			{
@@ -536,11 +703,11 @@ void UMRMeshComponent::SendBrickData_Internal(IMRMesh::FSendBrickDataArgs Args)
 
 void UMRMeshComponent::RemoveBodyInstance(int32 BodyIndex)
 {
-	BodyIds.RemoveAtSwap(BodyIndex);
-	BodySetups.RemoveAtSwap(BodyIndex);
 	BodyInstances[BodyIndex]->TermBody();
 	delete BodyInstances[BodyIndex];
 	BodyInstances.RemoveAtSwap(BodyIndex);
+	BodySetups.RemoveAtSwap(BodyIndex);
+	BodyIds.RemoveAtSwap(BodyIndex);
 }
 
 void UMRMeshComponent::OnUpdateTransform(EUpdateTransformFlags UpdateTransformFlags, ETeleportType Teleport)
diff --git a/Engine/Source/Runtime/MRMesh/Public/MRMeshComponent.h b/Engine/Source/Runtime/MRMesh/Public/MRMeshComponent.h
--- a/Engine/Source/Runtime/MRMesh/Public/MRMeshComponent.h
+++ b/Engine/Source/Runtime/MRMesh/Public/MRMeshComponent.h
@@ -49,7 +49,7 @@ public:
 
 
 
-UCLASS(meta = (BlueprintSpawnableComponent, Experimental), ClassGroup = Rendering)
+UCLASS(hideCategories=(Physics), meta = (BlueprintSpawnableComponent, Experimental), ClassGroup = Rendering)
 class MRMESH_API UMRMeshComponent : public UPrimitiveComponent, public IMRMesh
 {
 public:
@@ -58,6 +58,7 @@ public:
 	GENERATED_UCLASS_BODY()
 
 	virtual void BeginPlay() override;
+	void EndPlay(const EEndPlayReason::Type EndPlayReason) override;
 
 	UFUNCTION(BlueprintPure, Category = "Mesh Reconstruction")
 	bool IsConnected() const override { return bConnected; }
@@ -71,6 +72,19 @@ public:
 	UFUNCTION(BlueprintCallable, Category = "Mesh Reconstruction")
 	void Clear() override;
 
+protected:
+	virtual void OnActorEnableCollisionChanged() override;
+	virtual void UpdatePhysicsToRBChannels() override;
+public:
+	virtual void SetCollisionObjectType(ECollisionChannel Channel) override;
+	virtual void SetCollisionResponseToChannel(ECollisionChannel Channel, ECollisionResponse NewResponse) override;
+	virtual void SetCollisionResponseToAllChannels(ECollisionResponse NewResponse) override;
+	virtual void SetCollisionResponseToChannels(const FCollisionResponseContainer& NewResponses) override;
+	virtual void SetCollisionEnabled(ECollisionEnabled::Type NewType) override;
+	virtual void SetCollisionProfileName(FName InCollisionProfileName) override;
+
+	virtual void SetWalkableSlopeOverride(const FWalkableSlopeOverride& NewOverride) override;
+
 private:
 	//~ UPrimitiveComponent
 	virtual FPrimitiveSceneProxy* CreateSceneProxy() override;
@@ -78,8 +92,10 @@ private:
 	virtual FBoxSphereBounds CalcBounds(const FTransform& LocalToWorld) const override;
 	virtual void SetMaterial(int32 ElementIndex, class UMaterialInterface* InMaterial) override;
 	virtual bool DoCustomNavigableGeometryExport(FNavigableGeometryExport& GeomExport) const override;
+	virtual class UBodySetup* GetBodySetup();
 	//~ UPrimitiveComponent
 	//~ UActorComponent
+	virtual bool ShouldCreatePhysicsState() const override;
 	virtual void SendRenderDynamicData_Concurrent() override;
 	//~ UActorComponent
 
@@ -89,6 +105,8 @@ private:
 	//~ IMRMesh
 
 private:
+	void CacheBodySetupHelper();
+	class UBodySetup* CreateBodySetupHelper();
 	void SendBrickData_Internal(IMRMesh::FSendBrickDataArgs Args);
 
 	void RemoveBodyInstance(int32 BodyIndex);
@@ -104,13 +122,19 @@ private:
 	bool bCreateMeshProxySections = true;
 
 	/** If true, MRMesh will automatically update its navmesh whenever any Mesh section is updated. This may be expensive. If this is disabled use ForceNavMeshUpdate to trigger a navmesh update when necessary.  Moving the component will also trigger a navmesh update.*/
-	UPROPERTY(EditAnywhere, Category = Appearance)
+	UPROPERTY(EditAnywhere, Category = MRMesh)
 	bool bUpdateNavMeshOnMeshUpdate = true;
 
-	bool bEnableCollision = false;
+	/** If true, MRMesh will not create a collidable ridgid body for each mesh section and can therefore never have collision.  Avoids the cost of generating collision.*/
+	UPROPERTY(EditAnywhere, Category = MRMesh)
+	bool bNeverCreateCollisionMesh = false;
+
 	bool bConnected = false;
 
-	UPROPERTY(transient)
+	UPROPERTY(Transient)
+	class UBodySetup* CachedBodySetup = nullptr;
+
+	UPROPERTY(Transient)
 	TArray<UBodySetup*> BodySetups;
 
 	TArray<FBodyInstance*> BodyInstances;
