MockDataMeshTrackerComponent
-Created a component to generate mesh data for MRMesh so it can be tested without an AR device.

diff --git a/Engine/Source/Runtime/MRMesh/MRMesh.build.cs b/Engine/Source/Runtime/MRMesh/MRMesh.build.cs
--- a/Engine/Source/Runtime/MRMesh/MRMesh.build.cs
+++ b/Engine/Source/Runtime/MRMesh/MRMesh.build.cs
@@ -16,10 +16,16 @@ namespace UnrealBuildTool.Rules
 					"CoreUObject",
                     "Engine",
                     "RenderCore",
-                    "ShaderCore",
-                    "RHI"
+					"ShaderCore",
+                    "RHI",
+					"HeadMountedDisplay"
 				}
 			);
+
+			if (Target.bBuildEditor == true)
+			{
+				PrivateDependencyModuleNames.Add("UnrealEd");
+			}
 		}
 	}
 }
diff --git a/Engine/Source/Runtime/MRMesh/Private/MockDataMeshTrackerComponent.cpp b/Engine/Source/Runtime/MRMesh/Private/MockDataMeshTrackerComponent.cpp
new file mode 100644
--- /dev/null
+++ b/Engine/Source/Runtime/MRMesh/Private/MockDataMeshTrackerComponent.cpp
@@ -0,0 +1,625 @@
+// Copyright 1998-2018 Epic Games, Inc. All Rights Reserved.
+
+#include "MockDataMeshTrackerComponent.h"
+#include "Engine/Engine.h"
+#include "Engine/GameEngine.h"
+#include "GameFramework/WorldSettings.h"
+#include "HeadMountedDisplayFunctionLibrary.h"
+#if WITH_EDITOR
+#include "Editor.h"
+#endif
+#include "MRMeshComponent.h"
+
+
+
+class FMockDataMeshTrackerImpl
+{
+public:
+	FMockDataMeshTrackerImpl()
+		: MeshBrickIndex(0)
+	{
+		const int VertCount = 16;
+		FVector Verts[VertCount] =
+		{
+			FVector(0.0f,0.0f,0.0f),
+			FVector(10.0f,0.0f,0.0f),
+			FVector(20.0f,0.0f,0.0f),
+			FVector(30.0f,0.0f,0.0f),
+
+			FVector(0.0f,10.0f,0.0f),
+			FVector(10.0f,10.0f,10.0f),
+			FVector(20.0f,10.0f,10.0f),
+			FVector(30.0f,10.0f,0.0f),
+
+			FVector(0.0f,20.0f,0.0f),
+			FVector(10.0f,20.0f,10.0f),
+			FVector(20.0f,20.0f,10.0f),
+			FVector(30.0f,20.0f,0.0f),
+
+			FVector(0.0f,30.0f,0.0f),
+			FVector(10.0f,30.0f,0.0f),
+			FVector(20.0f,30.0f,0.0f),
+			FVector(30.0f,30.0f,0.0f),
+		};
+		FVector Normals[VertCount];
+		const FVector Center(0.15f, 0.15f, 0.0f);
+		for (int32 i = 0; i < VertCount; ++i)
+		{
+			Normals[i] = (Verts[i] - Center);
+			Normals[i].Normalize();
+		}
+		const int32 IndexCount = 54;
+		uint32 Indices[IndexCount] =
+		{
+			0,4,5,
+			0,5,1,
+			1,5,6,
+			1,6,2,
+			2,6,7,
+			2,7,3,
+
+			4,8,9,
+			4,9,5,
+			5,9,10,
+			5,10,6,
+			6,10,11,
+			6,11,7,
+
+			8,12,13,
+			8,13,9,
+			9,13,14,
+			9,14,10,
+			10,14,15,
+			10,15,11
+		};
+
+		const int NumBlocks = 4;
+		RawMockMeshData.AddDefaulted(NumBlocks);
+
+		for (int32 i = 0; i < NumBlocks; ++i)
+		{
+			FRawMockMeshData& Data = RawMockMeshData[i];
+			Data.Vertices.Reserve(VertCount);
+			Data.Normals.Reserve(VertCount);
+			for (int32 j = 0; j < VertCount; ++j)
+			{
+				Data.Vertices.Add(Verts[j]);			
+				Data.Normals.Add(Normals[j]);
+			}
+			Data.Indices.Reserve(IndexCount);
+			for (int32 j = 0; j < IndexCount; ++j)
+			{
+				Data.Indices.Add(Indices[j]);
+			}
+		}
+
+		// shift each block over by 0.3 to make a strip
+		for (int32 i = 0; i < NumBlocks; ++i)
+		{
+			FRawMockMeshData& Data = RawMockMeshData[i];
+			for (int32 j = 0; j < VertCount; ++j)
+			{
+				Data.Vertices[j].X += i * 30.0f;
+			}
+		}
+	};
+
+public:
+	// Next ID for bricks created with MR Mesh
+	int32 MeshBrickIndex = 0;
+
+	struct FRawMockMeshData
+	{
+		TArray<FVector> Vertices;
+		TArray<FVector> Normals;
+		TArray<uint32> Indices;
+	};
+	TArray<FRawMockMeshData> RawMockMeshData;
+
+	// Map of Raw mesh block IDs to MR Mesh brick IDs
+	TMap<uint32, uint64> MeshBrickCache;
+
+	// Keep a copy of the mesh data here.  MRMeshComponent will use it from the game and render thread.
+	struct FCachedMeshData
+	{
+		typedef TSharedPtr<FCachedMeshData, ESPMode::ThreadSafe> SharedPtr;
+		
+		FMockDataMeshTrackerImpl* Owner = nullptr;
+		
+		IMRMesh::FBrickId BrickId = 0;
+		TArray<FVector> OffsetVertices;
+		TArray<FVector> WorldVertices;
+		TArray<uint32> Triangles;
+		TArray<FVector> Normals;
+		TArray<FVector2D> UV0;
+		TArray<FColor> VertexColors;
+		TArray<FPackedNormal> Tangents;
+		TArray<float> Confidence;
+
+		void Recycle(SharedPtr& MeshData)
+		{
+			check(Owner);
+			FMockDataMeshTrackerImpl* TempOwner = Owner;
+			Owner = nullptr;
+	
+			BrickId = 0;
+			OffsetVertices.Reset();
+			WorldVertices.Reset();
+			Triangles.Reset();
+			Normals.Reset();
+			UV0.Reset();
+			VertexColors.Reset();
+			Tangents.Reset();
+			Confidence.Reset();
+	
+			TempOwner->FreeMeshDataCache(MeshData);
+		}
+
+		void Init(FMockDataMeshTrackerImpl* InOwner)
+		{
+			check(!Owner);
+			Owner = InOwner;
+		}
+	};
+
+
+	// This receipt will be kept in the FSendBrickDataArgs to ensure the cached data outlives MRMeshComponent use of it.
+	class FMeshTrackerComponentBrickDataReceipt : public IMRMesh::FBrickDataReceipt
+	{
+	public:
+		FMeshTrackerComponentBrickDataReceipt(FCachedMeshData::SharedPtr& MeshData) :
+			CachedMeshData(MeshData)
+		{
+		}
+		~FMeshTrackerComponentBrickDataReceipt() override
+		{
+			CachedMeshData->Recycle(CachedMeshData);
+		}
+	private:
+		FCachedMeshData::SharedPtr CachedMeshData;
+	};
+	
+	FCachedMeshData::SharedPtr AquireMeshDataCache()
+	{
+		if (FreeCachedMeshDatas.Num() > 0)
+		{
+			FScopeLock ScopeLock(&FreeCachedMeshDatasMutex);
+			FCachedMeshData::SharedPtr CachedMeshData(FreeCachedMeshDatas.Pop(false));
+			CachedMeshData->Init(this);
+			return CachedMeshData;
+		}
+		else
+		{
+			FCachedMeshData::SharedPtr CachedMeshData(new FCachedMeshData());
+			CachedMeshData->Init(this);
+			CachedMeshDatas.Add(CachedMeshData);
+			return CachedMeshData;
+		}
+	}
+
+	void FreeMeshDataCache(FCachedMeshData::SharedPtr& DataCache)
+	{
+		FScopeLock ScopeLock(&FreeCachedMeshDatasMutex);
+		FreeCachedMeshDatas.Add(DataCache);
+	}
+
+	FVector BoundsCenter;
+	FQuat BoundsRotation;
+
+	bool Create(const UMockDataMeshTrackerComponent& MeshTrackerComponent)
+	{
+		return true;
+	}
+
+	void Destroy()
+	{
+	}
+
+private:
+
+	// A free list to recycle the CachedMeshData instances.  
+	TArray<FCachedMeshData::SharedPtr> CachedMeshDatas;
+	TArray<FCachedMeshData::SharedPtr> FreeCachedMeshDatas;
+	FCriticalSection FreeCachedMeshDatasMutex; //The free list may be pushed/popped from multiple threads.
+};
+
+UMockDataMeshTrackerComponent::UMockDataMeshTrackerComponent(const FObjectInitializer& ObjectInitializer)
+	: Super(ObjectInitializer)
+	, VertexColorFromConfidenceZero(FLinearColor::Red)
+	, VertexColorFromConfidenceOne(FLinearColor::Blue)
+	, Impl(new FMockDataMeshTrackerImpl())
+{
+
+	// Make sure this component ticks
+	PrimaryComponentTick.bCanEverTick = true;
+	PrimaryComponentTick.bStartWithTickEnabled = true;
+	PrimaryComponentTick.TickGroup = TG_PrePhysics;
+
+	bAutoActivate = true;
+
+	BlockVertexColors.Add(FColor::Blue);
+	BlockVertexColors.Add(FColor::Red);
+	BlockVertexColors.Add(FColor::Green);
+	BlockVertexColors.Add(FColor::Yellow);
+	BlockVertexColors.Add(FColor::Cyan);
+	BlockVertexColors.Add(FColor::Magenta);
+
+#if WITH_EDITOR
+	if (GIsEditor)
+	{
+		FEditorDelegates::PrePIEEnded.AddUObject(this, &UMockDataMeshTrackerComponent::PrePIEEnded);
+	}
+#endif
+}
+
+UMockDataMeshTrackerComponent::~UMockDataMeshTrackerComponent()
+{
+	delete Impl;
+}
+
+void UMockDataMeshTrackerComponent::ConnectMRMesh(UMRMeshComponent* InMRMeshPtr)
+{
+	if (!InMRMeshPtr)
+	{
+		UE_LOG(LogMockMeshDataTracker, Warning,
+			TEXT("MRMesh given is not valid. Ignoring this connect."));
+		return;
+	}
+	else if (MRMesh)
+	{
+		UE_LOG(LogMockMeshDataTracker, Warning,
+			TEXT("MeshTrackerComponent already has a MRMesh connected.  Ignoring this connect."));
+		return;
+	}
+	else if (InMRMeshPtr->IsConnected())
+	{
+		UE_LOG(LogMockMeshDataTracker, Warning,
+			TEXT("MRMesh is already connected to a UMockDataMeshTrackerComponent. Ignoring this connect."));
+		return;
+	}
+	else
+	{
+		InMRMeshPtr->SetConnected(true);
+		MRMesh = InMRMeshPtr;
+	}
+}
+
+void UMockDataMeshTrackerComponent::DisconnectMRMesh(class UMRMeshComponent* InMRMeshPtr)
+{
+	if (!MRMesh)
+	{
+		UE_LOG(LogMockMeshDataTracker, Warning,
+			TEXT("MeshTrackerComponent MRMesh is already disconnected. Ignoring this disconnect."));
+		return;
+	}
+	else if (InMRMeshPtr != MRMesh)
+	{
+		UE_LOG(LogMockMeshDataTracker, Warning,
+			TEXT("MeshTrackerComponent MRMesh given is not the MRMesh connected. "
+				 "Ignoring this disconnect."));
+		return;
+	}
+	else
+	{
+		check(MRMesh->IsConnected());
+		MRMesh->SetConnected(false);
+	}
+	MRMesh = nullptr;
+}
+
+#if WITH_EDITOR
+void UMockDataMeshTrackerComponent::PostEditChangeProperty(FPropertyChangedEvent& e)
+{
+	if (e.Property != nullptr)
+	{
+		UE_LOG(LogMockMeshDataTracker, Log, TEXT("PostEditChangeProperty is changing MLMeshingSettings"));
+	}
+
+	Super::PostEditChangeProperty(e);
+}
+#endif
+
+void UMockDataMeshTrackerComponent::TickComponent(float DeltaTime, enum ELevelTick TickType, FActorComponentTickFunction *ThisTickFunction)
+{
+	Super::TickComponent(DeltaTime, TickType, ThisTickFunction);
+
+	if (!MRMesh)
+	{
+		return;
+	}
+
+	if (!Impl->Create(*this))
+	{
+		return;
+	}
+
+	//// Update the bounding box within which we scan for geometry.
+	//FTransform PoseInverse = UHeadMountedDisplayFunctionLibrary::GetTrackingToWorldTransform(this).Inverse();
+	//PoseInverse.ConcatenateRotation(BoundingVolume->GetComponentQuat());
+	//Impl->BoundsCenter = PoseInverse.TransformPosition(BoundingVolume->GetComponentLocation());
+	//Impl->BoundsRotation = PoseInverse.GetRotation();
+
+	//const float WorldToMetersScale = GWorld->GetWorldSettings()->WorldToMeters;
+
+	// Make sure MR Mesh is at 0,0,0 (verts received from meshing are in world space)
+	MRMesh->SendRelativeTransform(FTransform::Identity);
+
+	CurrentTime += DeltaTime;
+
+	if (ScanWorld && CurrentTime > LastUpdateTime + UpdateInterval)
+	{
+		LastUpdateTime = CurrentTime;
+		UpdateCount += 1;
+
+		static int MockDataPattern = 0;
+		if (MockDataPattern == 0)
+		{
+			// Cycle adding, updating, leaving alone, and removing blocks.
+
+			check(NumBlocks >= 3);
+			// Add one block
+			const int32 AddBlockIndex = (UpdateCount) % NumBlocks;
+			// Update one block
+			const int32 UpdateBlockIndex = (UpdateCount - 1) % NumBlocks;
+			// Remove oldest block
+			const int32 RemoveBlockIndex = (UpdateCount - NumBlocks + 1) % NumBlocks;
+
+			UE_LOG(LogMockMeshDataTracker, Log, TEXT("TickComponent is updating Add: %i Update: %i Remove: %i"), AddBlockIndex, UpdateBlockIndex, RemoveBlockIndex);
+
+			UpdateBlock(AddBlockIndex);
+			UpdateBlock(UpdateBlockIndex);
+			RemoveBlock(RemoveBlockIndex);
+		}
+		else
+		{
+			// Add then update 4 blocks.
+			UE_LOG(LogMockMeshDataTracker, Log, TEXT("TickComponent is adding 4 blocks"));
+			UpdateBlock(0);
+			UpdateBlock(1);
+			UpdateBlock(2);
+			UpdateBlock(3);
+		}
+
+		static bool bStopScanningEveryUpdate = false;
+		if (bStopScanningEveryUpdate)
+		{
+			ScanWorld = false;
+		}
+	}
+}
+
+void UMockDataMeshTrackerComponent::RemoveBlock(int32 BlockIndex)
+{
+	// Delete the brick and its cache entry
+	if (Impl->MeshBrickCache.Contains(BlockIndex))
+	{
+		const static TArray<FVector> EmptyVertices;
+		const static TArray<FVector2D> EmptyUVs;
+		const static TArray<FPackedNormal> EmptyTangents;
+		const static TArray<FColor> EmptyVertexColors;
+		const static TArray<uint32> EmptyTriangles;
+		const auto& BrickId = Impl->MeshBrickCache[BlockIndex];
+		static_cast<IMRMesh*>(MRMesh)->SendBrickData(IMRMesh::FSendBrickDataArgs
+			{
+				nullptr,
+				BrickId,
+				EmptyVertices,
+				EmptyUVs,
+				EmptyTangents,
+				EmptyVertexColors,
+				EmptyTriangles
+			}
+		);
+								
+		//if (OnMeshTrackerUpdated.IsBound())
+		//{
+		//	const static TArray<FVector> EmptyNormals;
+		//	const static TArray<float> EmptyConfidence;
+		//	const static TArray<int32> EmptyTriangles2;
+		//	OnMeshTrackerUpdated.Broadcast((uint32)BlockIndex, EmptyVertices, EmptyTriangles2, EmptyNormals, EmptyConfidence);
+		//}
+
+		Impl->MeshBrickCache.Remove(BlockIndex);
+	}
+}
+
+void UMockDataMeshTrackerComponent::UpdateBlock(int32 BlockIndex)
+{
+	// Create a brick index for any new mesh block
+	if (!Impl->MeshBrickCache.Contains(BlockIndex))
+	{
+		Impl->MeshBrickCache.Add(BlockIndex, Impl->MeshBrickIndex++);
+	}
+
+	const FMockDataMeshTrackerImpl::FRawMockMeshData& RawMeshData = Impl->RawMockMeshData[BlockIndex];
+
+	// Acquire mesh data cache and mark its brick ID
+	FMockDataMeshTrackerImpl::FCachedMeshData::SharedPtr CurrentMeshDataCache = Impl->AquireMeshDataCache();
+	const auto& BrickId = Impl->MeshBrickCache[BlockIndex];
+	CurrentMeshDataCache->BrickId = BrickId;
+
+	// Pull vertices
+	const FVector VertexOffset = UHeadMountedDisplayFunctionLibrary::GetTrackingToWorldTransform(this).Inverse().GetLocation();
+	const int32 VertexCount = RawMeshData.Vertices.Num();
+	CurrentMeshDataCache->OffsetVertices.Reserve(VertexCount);
+	CurrentMeshDataCache->WorldVertices.Reserve(VertexCount);
+	for (int32 v = 0; v < VertexCount; ++ v)
+	{
+			CurrentMeshDataCache->OffsetVertices.Add(RawMeshData.Vertices[v] - VertexOffset);
+			CurrentMeshDataCache->WorldVertices.Add(RawMeshData.Vertices[v]);
+	}
+
+	// Pull indices
+	const int32 IndexCount = RawMeshData.Indices.Num();
+	CurrentMeshDataCache->Triangles.Reserve(IndexCount);
+	for (int32 i = 0; i < IndexCount; ++ i)
+	{
+		CurrentMeshDataCache->Triangles.Add(static_cast<uint32>(RawMeshData.Indices[i]));
+	}
+
+	// Pull normals
+	CurrentMeshDataCache->Normals.Reserve(VertexCount);
+	if (RequestNormals)
+	{
+		for (int32 n = 0; n < VertexCount; ++ n)
+		{
+			CurrentMeshDataCache->Normals.Add(RawMeshData.Normals[n]);
+		}
+	}
+	// If no normals were provided we need to pack fake ones for Vulkan
+	else
+	{
+		for (int32 n = 0; n < VertexCount; ++ n)
+		{
+			FVector FakeNormal = CurrentMeshDataCache->OffsetVertices[n];
+			FakeNormal.Normalize();
+			CurrentMeshDataCache->Normals.Add(FakeNormal);
+		}
+	}
+
+	// Calculate and pack tangents
+	CurrentMeshDataCache->Tangents.Reserve(VertexCount * 2);
+	for (int32 t = 0; t < VertexCount; ++ t)
+	{
+		const FVector& Norm = CurrentMeshDataCache->Normals[t];
+
+		// Calculate tangent
+		auto Perp = Norm.X < Norm.Z ? 
+			FVector(1.0f, 0.0f, 0.0f) : FVector(0.0f, 1.0f, 0.0f);
+		auto Tang = FVector::CrossProduct(Norm, Perp);
+
+		CurrentMeshDataCache->Tangents.Add(Tang);
+		CurrentMeshDataCache->Tangents.Add(Norm);
+	}
+
+	// Pull confidence
+	if (RequestVertexConfidence)
+	{
+		CurrentMeshDataCache->Confidence.Reserve(VertexCount);
+		const float Confidence = (float)BlockIndex / (float)NumBlocks;
+		for (int32 c = 0; c < VertexCount; ++c)
+		{
+			CurrentMeshDataCache->Confidence.Add(Confidence);
+		}
+	}
+
+	// Apply chosen vertex color mode
+	switch (VertexColorMode)
+	{
+		case EMeshTrackerVertexColorMode::Confidence:
+		{
+			if (RequestVertexConfidence)
+			{
+				CurrentMeshDataCache->VertexColors.Reserve(VertexCount);
+				for (int32 v = 0; v < VertexCount; ++ v)
+				{
+					const FLinearColor VertexColor = FMath::Lerp(VertexColorFromConfidenceZero, 
+						VertexColorFromConfidenceOne, CurrentMeshDataCache->Confidence[v]);
+					CurrentMeshDataCache->VertexColors.Add(VertexColor.ToFColor(false));
+				}
+			}
+			else
+			{
+				UE_LOG(LogMockMeshDataTracker, Warning, TEXT("MeshTracker vertex color mode is Confidence "
+					"but no confidence values available. Using white for all blocks."));
+			}
+			break;
+		}
+		case EMeshTrackerVertexColorMode::Block:
+		{
+			if (BlockVertexColors.Num() > 0)
+			{
+				const FColor& VertexColor = BlockVertexColors[BlockIndex % BlockVertexColors.Num()];
+
+				CurrentMeshDataCache->VertexColors.Reserve(VertexCount);
+				for (int32 v = 0; v < VertexCount; ++ v)
+				{
+					CurrentMeshDataCache->VertexColors.Add(VertexColor);
+				}
+			}
+			else
+			{
+				UE_LOG(LogMockMeshDataTracker, Warning, TEXT("MeshTracker vertex color mode is Block but "
+					"no BlockVertexColors set. Using white for all blocks."));
+			}
+			break;
+		}
+		case EMeshTrackerVertexColorMode::None:
+		{
+			break;
+		}
+		default:
+			check(false);
+			break;
+	}
+
+	// To work in all rendering paths we always set a vertex color
+	if (CurrentMeshDataCache->VertexColors.Num() == 0)
+	{
+		CurrentMeshDataCache->VertexColors.Reserve(VertexCount);
+		for (int32 v = 0; v < VertexCount; ++ v)
+		{
+			CurrentMeshDataCache->VertexColors.Add(FColor::White);
+		}
+	}
+
+	// Write UVs
+	CurrentMeshDataCache->UV0.Reserve(VertexCount);
+	for (int32 v = 0; v < VertexCount; ++ v)
+	{
+		const float FakeCoord = static_cast<float>(v) / static_cast<float>(VertexCount);
+		CurrentMeshDataCache->UV0.Add(FVector2D(FakeCoord, FakeCoord));
+	}
+
+	// Create/update brick
+	static_cast<IMRMesh*>(MRMesh)->SendBrickData(IMRMesh::FSendBrickDataArgs
+		{
+			TSharedPtr<IMRMesh::FBrickDataReceipt, ESPMode::ThreadSafe>
+				(new FMockDataMeshTrackerImpl::FMeshTrackerComponentBrickDataReceipt(CurrentMeshDataCache)),
+			CurrentMeshDataCache->BrickId,
+			CurrentMeshDataCache->WorldVertices,
+			CurrentMeshDataCache->UV0,
+			CurrentMeshDataCache->Tangents,
+			CurrentMeshDataCache->VertexColors,
+			CurrentMeshDataCache->Triangles
+		}
+	);
+
+	// Broadcast that a mesh was updated
+	if (OnMeshTrackerUpdated.IsBound())
+	{
+		// Hack because blueprints don't support uint32.
+		TArray<int32> Triangles(reinterpret_cast<const int32*>(CurrentMeshDataCache->
+			Triangles.GetData()), CurrentMeshDataCache->Triangles.Num());
+		OnMeshTrackerUpdated.Broadcast(CurrentMeshDataCache->BrickId, CurrentMeshDataCache->OffsetVertices, 
+			Triangles, CurrentMeshDataCache->Normals, CurrentMeshDataCache->Confidence);
+	}
+}
+
+void UMockDataMeshTrackerComponent::BeginDestroy()
+{
+	if (MRMesh != nullptr)
+	{
+		DisconnectMRMesh(MRMesh);
+	}
+	Super::BeginDestroy();
+}
+
+void UMockDataMeshTrackerComponent::FinishDestroy()
+{
+#if WITH_EDITOR
+	if (GIsEditor)
+	{
+		FEditorDelegates::PrePIEEnded.RemoveAll(this);
+	}
+#endif
+	Impl->Destroy();
+	Super::FinishDestroy();
+}
+
+#if WITH_EDITOR
+void UMockDataMeshTrackerComponent::PrePIEEnded(bool bWasSimulatingInEditor)
+{
+	Impl->Destroy();
+}
+#endif
diff --git a/Engine/Source/Runtime/MRMesh/Public/MockDataMeshTrackerComponent.h b/Engine/Source/Runtime/MRMesh/Public/MockDataMeshTrackerComponent.h
new file mode 100644
--- /dev/null
+++ b/Engine/Source/Runtime/MRMesh/Public/MockDataMeshTrackerComponent.h
@@ -0,0 +1,149 @@
+// Copyright 1998-2018 Epic Games, Inc. All Rights Reserved.
+
+#pragma once
+
+#include "Components/SceneComponent.h"
+#include "Delegates/DelegateCombinations.h"
+#include "MockDataMeshTrackerComponent.generated.h"
+
+DEFINE_LOG_CATEGORY_STATIC(LogMockMeshDataTracker, Log, All);
+
+
+/** Vertex color mode. */
+UENUM(BlueprintType)
+enum class EMeshTrackerVertexColorMode : uint8
+{
+	/** Vertex Color is not set. */
+	None    UMETA(DisplayName = "No Vertex Color"),
+	/*! Vertex confidence is interpolated between two specified colors. */
+	Confidence  UMETA(DisplayName = "Vertex Confidence"),
+	/*! Each block is given a color from a list. */
+	Block  UMETA(DisplayName = "Blocks Colored")
+};
+
+/**
+	The MeshTrackerComponent class manages requests for environmental mesh data, processes the results and provides
+	them to the calling system. The calling system is able request environmental mesh data within a specified area.
+	Various other search criteria can be set via this class's public properties.  Mesh data requests are processed
+	on a separate thread.  Once a mesh data request has been processed the calling system will be notified via an
+	FOnMeshTrackerUpdated broadcast.
+*/
+UCLASS(ClassGroup = Rendering, BlueprintType, Blueprintable, EditInlineNew, meta = (BlueprintSpawnableComponent))
+class MRMESH_API UMockDataMeshTrackerComponent
+	: public USceneComponent
+{
+	GENERATED_UCLASS_BODY()
+
+public:
+	/** Destroys the FMeshTrackerImpl instance.*/
+	~UMockDataMeshTrackerComponent();
+
+	/**
+		Sets the procedural mesh component that will store and display the environmental mesh results.
+		@param InMRMeshPtr The procedural mesh component to store the query result in.
+	*/
+	UFUNCTION(BlueprintCallable, Category = "Mesh Reconstruction")
+	void ConnectMRMesh(class UMRMeshComponent* InMRMeshPtr);
+
+	/**
+		Unlinks the current procedural mesh component from the mesh tracking system.
+		@param InMRMeshPtr The procedural mesh component to unlink from the mesh tracking system.
+	*/
+	UFUNCTION(BlueprintCallable, Category = "Mesh Reconstruction")
+	void DisconnectMRMesh(class UMRMeshComponent* InMRMeshPtr);
+
+	/**
+	Delegate used by OnMeshUpdated().
+	@param Index The index of the mesh section in the ProceduralMeshComponent that was updated.
+	@param Vertices List of all vertices in the updated mesh section.
+	@param Triangles List of all triangles in the updated mesh section.
+	@param Normals List of the normals of all triangles in the updated mesh section.
+	@param Confidence List of the confidence values per vertex in the updated mesh section. This can be used to determine if the user needs to scan more.
+	*/
+	DECLARE_DYNAMIC_MULTICAST_DELEGATE_FiveParams(FOnMockDataMeshTrackerUpdated, int32, Index, const TArray<FVector>&, Vertices, const TArray<int32>&, Triangles, const TArray<FVector>&, Normals, const TArray<float>&, Confidence);
+
+	/** Activated whenever new information about this mesh tracker is detected. */
+	UPROPERTY(BlueprintAssignable)
+	FOnMockDataMeshTrackerUpdated OnMeshTrackerUpdated;
+
+	/** Set to true to start scanning the world for meshes. */
+	UPROPERTY(EditAnywhere, BlueprintReadWrite, Category = "Meshing")
+	bool ScanWorld = true;
+
+	/** If true, the system will generate normals for the triangle vertices. */
+	UPROPERTY(EditAnywhere, BlueprintReadWrite, Category = "Meshing")
+	bool RequestNormals = true;
+
+	/**
+		If true, the system will generate the mesh confidence values for the triangle vertices.
+		These confidence values can be used to determine if the user needs to scan more.
+	*/
+	UPROPERTY(EditAnywhere, BlueprintReadWrite, Category = "Meshing")
+	bool RequestVertexConfidence = false;
+
+	/**
+	* Vertex Colors can be unused, or filled with several types of information.
+	*/
+	UPROPERTY(EditAnywhere, BlueprintReadWrite, Category = "Meshing")
+		EMeshTrackerVertexColorMode VertexColorMode = EMeshTrackerVertexColorMode::None;
+
+	/** Colors through which we cycle when setting vertex color by block. */
+	UPROPERTY(EditAnywhere, BlueprintReadWrite, Category = "Meshing")
+	TArray<FColor> BlockVertexColors;
+
+	/** Color mapped to confidence value of zero. */
+	UPROPERTY(EditAnywhere, BlueprintReadWrite, Category = "Meshing")
+	FLinearColor VertexColorFromConfidenceZero;
+
+	/** Color mapped to confidence value of one. */
+	UPROPERTY(EditAnywhere, BlueprintReadWrite, Category = "Meshing")
+	FLinearColor VertexColorFromConfidenceOne;
+
+
+
+	/** Update Interval in Seconds. */
+	UPROPERTY(EditAnywhere, BlueprintReadWrite, Category = "Meshing|MockData")
+	float UpdateInterval = 3.0f;
+
+
+
+	/** MRMeshComponent can render and provide collision based on the Mesh data. */
+	UPROPERTY(transient)
+	class UMRMeshComponent* MRMesh;
+
+	/**
+		The procedural mesh generated should bake physics and update its collision data.
+		TODO: Can we calculate this based on the collision flags of the UProceduralMeshComponent?
+		It seems like those are not being respected directly by the component.
+	*/
+
+	/** Polls for and handles the results of the environmental mesh queries. */
+	virtual void TickComponent(float DeltaTime, enum ELevelTick TickType, FActorComponentTickFunction *ThisTickFunction) override;
+
+	/** Unlinks the current procedural mesh component from the mesh tracking system. */
+	virtual void BeginDestroy() override;
+
+	/** Destroys the interface object to the mesh tracking api*/
+	virtual void FinishDestroy() override;
+
+#if WITH_EDITOR
+	virtual void PostEditChangeProperty(FPropertyChangedEvent& e) override;
+#endif
+
+private:
+	class FMockDataMeshTrackerImpl *Impl;
+
+	float LastUpdateTime = 0.0f;
+	float CurrentTime = 0.0f;
+	int32 UpdateCount = 0;
+	int32 NumBlocks = 4;
+
+	void UpdateBlock(int32 BlockIndex);
+	void RemoveBlock(int32 BlockIndex);
+
+
+#if WITH_EDITOR
+private:
+	void PrePIEEnded(bool bWasSimulatingInEditor);
+#endif
+};
