FMRMeshProxy::GetViewRelevance needs to flag translucency if the material is a translucent material, otherwise no draw will be added for the mesh in  FBasePassMeshProcessor::AddMeshBatch.

diff --git a/Engine/Source/Runtime/MRMesh/Private/MRMeshComponent.cpp b/Engine/Source/Runtime/MRMesh/Private/MRMeshComponent.cpp
--- a/Engine/Source/Runtime/MRMesh/Private/MRMeshComponent.cpp
+++ b/Engine/Source/Runtime/MRMesh/Private/MRMeshComponent.cpp
@@ -352,6 +352,7 @@ private:
 		Result.bRenderInMainPass = ShouldRenderInMainPass();
 		Result.bUsesLightingChannels = GetLightingChannelMask() != GetDefaultLightingChannelMask();
 		Result.bRenderCustomDepth = ShouldRenderCustomDepth();
+		Result.bSeparateTranslucencyRelevance = MaterialToUse->GetMaterial()->bEnableSeparateTranslucency;
 		//MaterialRelevance.SetPrimitiveViewRelevance(Result);
 		return Result;
 	}
