Merge MRMesh changes from Ringtail

diff --git a/Engine/Source/Runtime/MRMesh/Private/MRMeshComponent.cpp b/Engine/Source/Runtime/MRMesh/Private/MRMeshComponent.cpp
--- a/Engine/Source/Runtime/MRMesh/Private/MRMeshComponent.cpp
+++ b/Engine/Source/Runtime/MRMesh/Private/MRMeshComponent.cpp
@@ -22,6 +22,7 @@
 #include "Engine/Engine.h"
 #include "UObject/UObjectThreadContext.h"
 #include "Stats/Stats.h"
+#include "UObject/ConstructorHelpers.h"
 
 DECLARE_CYCLE_STAT(TEXT("MrMesh SetCollisionProfileName"), STAT_MrMesh_SetCollisionProfileName, STATGROUP_Physics);
 
@@ -81,6 +82,19 @@ public:
 		FMemory::Memcpy(Buffer, Indices.GetData(), Indices.Num() * sizeof(int32));
 		RHIUnlockIndexBuffer(IndexBufferRHI);
 	}
+
+	void InitRHIWith(const TArray<uint16>& Indices)
+	{
+		NumIndices = Indices.Num();
+
+		FRHIResourceCreateInfo CreateInfo;
+		void* Buffer = nullptr;
+		IndexBufferRHI = RHICreateAndLockIndexBuffer(sizeof(uint16), Indices.Num() * sizeof(uint16), BUF_Static, CreateInfo, Buffer);
+
+		// Write the indices to the index buffer.
+		FMemory::Memcpy(Buffer, Indices.GetData(), Indices.Num() * sizeof(uint16));
+		RHIUnlockIndexBuffer(IndexBufferRHI);
+	}
 };
 
 
@@ -92,12 +106,15 @@ struct FMRMeshProxySection
 	IMRMesh::FBrickId BrickId;
 	/** Position buffer */
 	FMRMeshVertexBuffer<FVector> PositionBuffer;
+	// AR meshes from HoloLens and ARKit don't have these
+#if !(PLATFORM_HOLOLENS || PLATFORM_IOS)
 	/** Texture coordinates buffer */
 	FMRMeshVertexBuffer<FVector2D> UVBuffer;
 	/** Tangent space buffer */
 	FMRMeshVertexBuffer<FPackedNormal> TangentXZBuffer;
 	/** We don't need color */
 	FMRMeshVertexBuffer<FColor> ColorBuffer;
+#endif
 	/** Index buffer for this section */
 	FMRMeshIndexBuffer IndexBuffer;
 	/** Vertex factory for this section */
@@ -117,9 +134,12 @@ struct FMRMeshProxySection
 	void ReleaseResources()
 	{
 		PositionBuffer.ReleaseResource();
+		// AR meshes from HoloLens and ARKit don't have these
+#if !(PLATFORM_HOLOLENS || PLATFORM_IOS)
 		UVBuffer.ReleaseResource();
 		TangentXZBuffer.ReleaseResource();
 		ColorBuffer.ReleaseResource();
+#endif
 		IndexBuffer.ReleaseResource();
 		VertexFactory.ReleaseResource();
 	}
@@ -143,6 +163,8 @@ static void InitVertexFactory(FLocalVertexFactory* VertexFactory, const FMRMeshP
 			NewData.PositionComponent = FVertexStreamComponent(&MRMeshSection.PositionBuffer, 0, sizeof(FVector), VET_Float3, EVertexStreamUsage::Default);
 		}
 
+		// AR meshes from HoloLens and ARKit don't have these
+#if !(PLATFORM_HOLOLENS || PLATFORM_IOS)
 		if (MRMeshSection.UVBuffer.NumVerts != 0)
 		{
 			NewData.TextureCoordinatesSRV = MRMeshSection.UVBufferSRV;
@@ -162,6 +184,7 @@ static void InitVertexFactory(FLocalVertexFactory* VertexFactory, const FMRMeshP
 			NewData.ColorComponentsSRV = MRMeshSection.ColorBufferSRV;
 			NewData.ColorComponent = FVertexStreamComponent(&MRMeshSection.ColorBuffer, 0, sizeof(FColor), VET_Color, EVertexStreamUsage::ManualFetch);
 		}
+#endif
 
 		VertexFactory->SetData(NewData);
 		VertexFactory->InitResource();
@@ -180,7 +203,7 @@ public:
 
 	FMRMeshProxy(const UMRMeshComponent* InComponent)
 	: FPrimitiveSceneProxy(InComponent, InComponent->GetFName())
-	, MaterialToUse((InComponent->Material!=nullptr) ? InComponent->Material : UMaterial::GetDefaultMaterial(MD_Surface) )
+	, MaterialToUse((InComponent->Material != nullptr) ? InComponent->Material : (InComponent->WireframeMaterial) ? InComponent->WireframeMaterial : UMaterial::GetDefaultMaterial(MD_Surface))
 	, FeatureLevel(GetScene().GetFeatureLevel())
 	{
 	}
@@ -204,11 +227,14 @@ public:
 		FMRMeshProxySection* NewSection = new FMRMeshProxySection(Args.BrickId, FeatureLevel);
 		ProxySections.Add(NewSection);
 
+		// AR meshes from HoloLens and ARKit don't have these
+#if !(PLATFORM_HOLOLENS || PLATFORM_IOS)
 		// Vulkan requires that all the buffers be full.
 		const int32 NumVerts = Args.PositionData.Num();
 		check((NumVerts == Args.ColorData.Num()));
 		check((NumVerts == Args.UVData.Num()));
 		check((NumVerts * 2) == Args.TangentXZData.Num());
+#endif
 
 		// POSITION BUFFER
 		{
@@ -217,6 +243,8 @@ public:
 			NewSection->PositionBufferSRV = RHICreateShaderResourceView(NewSection->PositionBuffer.VertexBufferRHI, sizeof(float), PF_R32_FLOAT);
 		}
 
+		// AR meshes from HoloLens and ARKit don't have these
+#if !(PLATFORM_HOLOLENS || PLATFORM_IOS)
 		// TEXTURE COORDS BUFFER
 		{
 			NewSection->UVBuffer.InitResource();
@@ -250,6 +278,7 @@ public:
 				NewSection->ColorBufferSRV = RHICreateShaderResourceView(NewSection->ColorBuffer.VertexBufferRHI, 4, PF_R8G8B8A8);
 			}
 		}
+#endif
 
 		// INDEX BUFFER
 		{
@@ -376,6 +405,8 @@ private:
 UMRMeshComponent::UMRMeshComponent(const FObjectInitializer& ObjectInitializer)
 : Super(ObjectInitializer)
 {
+	static ConstructorHelpers::FObjectFinder<UMaterialInterface> WireframeMaterialRef(TEXT("/Engine/EngineDebugMaterials/WireframeMaterial"));
+	WireframeMaterial = WireframeMaterialRef.Object;
 }
 
 void UMRMeshComponent::BeginPlay()
@@ -532,6 +563,10 @@ void UMRMeshComponent::GetUsedMaterials(TArray<UMaterialInterface*>& OutMaterial
 	{
 		OutMaterials.Add(Material);
 	}
+	if (WireframeMaterial != nullptr)
+	{
+		OutMaterials.Add(WireframeMaterial);
+	}
 }
 
 FBoxSphereBounds UMRMeshComponent::CalcBounds(const FTransform& LocalToWorld) const
@@ -808,3 +843,48 @@ void UMRMeshComponent::Clear()
 	ClearAllBrickData();
 	UE_LOG(LogMrMesh, Log, TEXT("Clearing all brick data"));
 }
+
+struct FMeshArrayHolder :
+	public IMRMesh::FBrickDataReceipt
+{
+	TArray<FVector> Vertices;
+	TArray<MRMESH_INDEX_TYPE> Indices;
+	// Super wasteful of memory and perf, but the vertex factory requires these to be filled
+	// @todo Write a vertex factory that doesn't need all this overhead
+	TArray<FVector2D> BogusUVs;
+	TArray<FPackedNormal> BogusTangents;
+	TArray<FColor> BogusColors;
+
+	FMeshArrayHolder(TArray<FVector>& InVertices, TArray<MRMESH_INDEX_TYPE>& InIndices)
+		: Vertices(MoveTemp(InVertices))
+		, Indices(MoveTemp(InIndices))
+	{
+		int32 CurrentNumVertices = Vertices.Num();
+
+		BogusUVs.AddZeroed(CurrentNumVertices);
+		BogusColors.AddZeroed(CurrentNumVertices);
+		BogusTangents.AddZeroed(CurrentNumVertices * 2);
+	}
+};
+
+void UMRMeshComponent::UpdateTransformAndMesh(const FTransform& InTransform, TArray<FVector>& Vertices, TArray<MRMESH_INDEX_TYPE>& Indices)
+{
+	SetComponentToWorld(InTransform);
+
+	// Create our struct that will hold the data until the render thread is done with it
+	TSharedPtr<FMeshArrayHolder, ESPMode::ThreadSafe> MeshHolder(new FMeshArrayHolder(Vertices, Indices));
+	// NOTE: Vertices and Indices are empty due to MoveTemp()!!!
+
+	SendBrickData_Internal(IMRMesh::FSendBrickDataArgs
+		{
+			MeshHolder,
+			0,
+			MeshHolder->Vertices,
+			MeshHolder->BogusUVs,
+			MeshHolder->BogusTangents,
+			MeshHolder->BogusColors,
+			MeshHolder->Indices
+		}
+	);
+}
+
diff --git a/Engine/Source/Runtime/MRMesh/Private/MockDataMeshTrackerComponent.cpp b/Engine/Source/Runtime/MRMesh/Private/MockDataMeshTrackerComponent.cpp
--- a/Engine/Source/Runtime/MRMesh/Private/MockDataMeshTrackerComponent.cpp
+++ b/Engine/Source/Runtime/MRMesh/Private/MockDataMeshTrackerComponent.cpp
@@ -49,7 +49,7 @@ public:
 			Normals[i].Normalize();
 		}
 		const int32 IndexCount = 54;
-		uint32 Indices[IndexCount] =
+		MRMESH_INDEX_TYPE Indices[IndexCount] =
 		{
 			0,4,5,
 			0,5,1,
@@ -112,7 +112,7 @@ public:
 	{
 		TArray<FVector> Vertices;
 		TArray<FVector> Normals;
-		TArray<uint32> Indices;
+		TArray<MRMESH_INDEX_TYPE> Indices;
 	};
 	TArray<FRawMockMeshData> RawMockMeshData;
 
@@ -129,11 +129,11 @@ public:
 		IMRMesh::FBrickId BrickId = 0;
 		TArray<FVector> OffsetVertices;
 		TArray<FVector> WorldVertices;
-		TArray<uint32> Triangles;
 		TArray<FVector> Normals;
 		TArray<FVector2D> UV0;
 		TArray<FColor> VertexColors;
 		TArray<FPackedNormal> Tangents;
+		TArray<MRMESH_INDEX_TYPE> Triangles;
 		TArray<float> Confidence;
 
 		void Recycle(SharedPtr& MeshData)
@@ -397,7 +397,7 @@ void UMockDataMeshTrackerComponent::RemoveBlock(int32 BlockIndex)
 		const static TArray<FVector2D> EmptyUVs;
 		const static TArray<FPackedNormal> EmptyTangents;
 		const static TArray<FColor> EmptyVertexColors;
-		const static TArray<uint32> EmptyTriangles;
+		const static TArray<MRMESH_INDEX_TYPE> EmptyTriangles;
 		const auto& BrickId = Impl->MeshBrickCache[BlockIndex];
 		static_cast<IMRMesh*>(MRMesh)->SendBrickData(IMRMesh::FSendBrickDataArgs
 			{
@@ -454,7 +454,7 @@ void UMockDataMeshTrackerComponent::UpdateBlock(int32 BlockIndex)
 	CurrentMeshDataCache->Triangles.Reserve(IndexCount);
 	for (int32 i = 0; i < IndexCount; ++ i)
 	{
-		CurrentMeshDataCache->Triangles.Add(static_cast<uint32>(RawMeshData.Indices[i]));
+		CurrentMeshDataCache->Triangles.Add(static_cast<MRMESH_INDEX_TYPE>(RawMeshData.Indices[i]));
 	}
 
 	// Pull normals
@@ -588,11 +588,14 @@ void UMockDataMeshTrackerComponent::UpdateBlock(int32 BlockIndex)
 	// Broadcast that a mesh was updated
 	if (OnMeshTrackerUpdated.IsBound())
 	{
-		// Hack because blueprints don't support uint32.
-		TArray<int32> Triangles(reinterpret_cast<const int32*>(CurrentMeshDataCache->
-			Triangles.GetData()), CurrentMeshDataCache->Triangles.Num());
-		OnMeshTrackerUpdated.Broadcast(CurrentMeshDataCache->BrickId, CurrentMeshDataCache->OffsetVertices, 
-			Triangles, CurrentMeshDataCache->Normals, CurrentMeshDataCache->Confidence);
+		if (sizeof(MRMESH_INDEX_TYPE) == sizeof(uint32))
+		{
+			// Hack because blueprints don't support uint32.
+			TArray<int32> Triangles(reinterpret_cast<const int32*>(CurrentMeshDataCache->
+				Triangles.GetData()), CurrentMeshDataCache->Triangles.Num());
+			OnMeshTrackerUpdated.Broadcast(CurrentMeshDataCache->BrickId, CurrentMeshDataCache->OffsetVertices,
+				Triangles, CurrentMeshDataCache->Normals, CurrentMeshDataCache->Confidence);
+		}
 	}
 }
 
diff --git a/Engine/Source/Runtime/MRMesh/Public/MRMeshBufferDefines.h b/Engine/Source/Runtime/MRMesh/Public/MRMeshBufferDefines.h
new file mode 100644
--- /dev/null
+++ b/Engine/Source/Runtime/MRMesh/Public/MRMeshBufferDefines.h
@@ -0,0 +1,43 @@
+// Copyright 1998-2019 Epic Games, Inc. All Rights Reserved.
+
+#pragma once
+
+#ifndef WANTS_MRMESH_UVS
+	#if PLATFORM_WINDOWS
+		#define WANTS_MRMESH_UVS 1
+	#elif PLATFORM_HOLOLENS
+		#define WANTS_MRMESH_UVS 0
+	#elif PLATFORM_IOS
+		#define WANTS_MRMESH_UVS 0
+	#endif
+#endif
+
+#ifndef WANTS_MRMESH_TANGENTS
+	#if PLATFORM_WINDOWS
+		#define WANTS_MRMESH_TANGENTS 1
+	#elif PLATFORM_HOLOLENS
+		#define WANTS_MRMESH_TANGENTS 0
+	#elif PLATFORM_IOS
+		#define WANTS_MRMESH_TANGENTS 0
+	#endif
+#endif
+
+#ifndef WANTS_MRMESH_COLORS
+	#if PLATFORM_WINDOWS
+		#define WANTS_MRMESH_COLORS 1
+	#elif PLATFORM_HOLOLENS
+		#define WANTS_MRMESH_COLORS 0
+	#elif PLATFORM_IOS
+		#define WANTS_MRMESH_COLORS 0
+	#endif
+#endif
+
+#ifndef MRMESH_INDEX_TYPE
+	#if PLATFORM_WINDOWS
+		#define MRMESH_INDEX_TYPE uint32
+	#elif PLATFORM_HOLOLENS
+		#define MRMESH_INDEX_TYPE uint16
+	#elif PLATFORM_IOS
+		#define MRMESH_INDEX_TYPE uint16
+	#endif
+#endif
diff --git a/Engine/Source/Runtime/MRMesh/Public/MRMeshComponent.h b/Engine/Source/Runtime/MRMesh/Public/MRMeshComponent.h
--- a/Engine/Source/Runtime/MRMesh/Public/MRMeshComponent.h
+++ b/Engine/Source/Runtime/MRMesh/Public/MRMeshComponent.h
@@ -17,6 +17,18 @@ struct FDynamicMeshVertex;
 DECLARE_STATS_GROUP(TEXT("MRMesh"), STATGROUP_MRMESH, STATCAT_Advanced);
 DEFINE_LOG_CATEGORY_STATIC(LogMrMesh, Warning, All);
 
+
+//@todo JoeG - remove
+#ifndef PLATFORM_HOLOLENS
+#define PLATFORM_HOLOLENS 0
+#endif
+
+#if (PLATFORM_HOLOLENS || PLATFORM_IOS)
+	#define MRMESH_INDEX_TYPE uint16
+#else
+	#define MRMESH_INDEX_TYPE uint32
+#endif
+
 class IMRMesh
 {
 public:
@@ -35,7 +47,7 @@ public:
 		const TArray<FVector2D>& UVData;
 		const TArray<FPackedNormal>& TangentXZData;
 		const TArray<FColor>& ColorData;
-		const TArray<uint32>& Indices;
+		const TArray<MRMESH_INDEX_TYPE>& Indices;
 	};
 
 	virtual void SetConnected(bool value) = 0;
@@ -74,6 +86,10 @@ public:
 
 	// UPrimitiveComponent.. public BP function needs to stay public to avoid nativization errors. (RR)
 	virtual void SetMaterial(int32 ElementIndex, class UMaterialInterface* InMaterial) override;
+
+	/** Updates from HoloLens or iOS */
+	void UpdateTransformAndMesh(const FTransform& InTransform, TArray<FVector>& Vertices, TArray<MRMESH_INDEX_TYPE>& Indices);
+
 protected:
 	virtual void OnActorEnableCollisionChanged() override;
 	virtual void UpdatePhysicsToRBChannels() override;
@@ -138,6 +154,9 @@ private:
 	UPROPERTY(Transient)
 	TArray<UBodySetup*> BodySetups;
 
+	UPROPERTY()
+	UMaterialInterface* WireframeMaterial;
+
 	TArray<FBodyInstance*> BodyInstances;
 	TArray<IMRMesh::FBrickId> BodyIds;
 };
\ No newline at end of file
