Merge: More MRMesh changes and session config properties

diff --git a/Engine/Source/Runtime/AugmentedReality/Public/ARSessionConfig.h b/Engine/Source/Runtime/AugmentedReality/Public/ARSessionConfig.h
--- a/Engine/Source/Runtime/AugmentedReality/Public/ARSessionConfig.h
+++ b/Engine/Source/Runtime/AugmentedReality/Public/ARSessionConfig.h
@@ -230,6 +230,28 @@ public:
 	bool ShouldDoVerticalPlaneDetection() const { return bVerticalPlaneDetection; }
 	
 	const TArray<uint8>& GetSerializedARCandidateImageDatabase() const;	
+
+
+	/** Whether the AR system should generate mesh data that can be rendered, collided against, nav mesh generated on, etc. */
+	UPROPERTY(EditAnywhere, Category = "AR Settings | World Mapping")
+	bool bGenerateMeshDataFromTrackedGeometry;
+
+	/** Whether the AR system should generate collision data from the mesh data or not */
+	UPROPERTY(EditAnywhere, Category = "AR Settings | World Mapping")
+	bool bGenerateCollisionForMeshData;
+
+	/** Whether the AR system should generate navigation mesh data from the mesh data or not */
+	UPROPERTY(EditAnywhere, Category = "AR Settings | World Mapping")
+	bool bGenerateNavMeshForMeshData;
+
+	/** Whether the AR system render the mesh data as occlusion meshes or not */
+	UPROPERTY(EditAnywhere, Category = "AR Settings | World Mapping")
+	bool bUseMeshDataForOcclusion;
+
+	/** Whether the AR system should render the mesh data in wireframe or not */
+	UPROPERTY(EditAnywhere, Category = "AR Settings | World Mapping")
+	bool bRenderMeshDataInWireframe;
+
 private:
 	//~ UObject interface
 	virtual void Serialize(FArchive& Ar) override;
@@ -297,7 +319,7 @@ protected:
 	EAREnvironmentCaptureProbeType EnvironmentCaptureProbeType;
 
 	/** A previously saved world that is to be loaded when the session starts */
-	UPROPERTY(VisibleAnywhere, Category="AR Settings")
+	UPROPERTY(VisibleAnywhere, Category="AR Settings | World Mapping")
 	TArray<uint8> WorldMapData;
 
 	/** A list of candidate objects to search for in the scene */
diff --git a/Engine/Source/Runtime/MRMesh/Private/MRMeshComponent.cpp b/Engine/Source/Runtime/MRMesh/Private/MRMeshComponent.cpp
--- a/Engine/Source/Runtime/MRMesh/Private/MRMeshComponent.cpp
+++ b/Engine/Source/Runtime/MRMesh/Private/MRMeshComponent.cpp
@@ -106,15 +106,12 @@ struct FMRMeshProxySection
 	IMRMesh::FBrickId BrickId;
 	/** Position buffer */
 	FMRMeshVertexBuffer<FVector> PositionBuffer;
-	// AR meshes from HoloLens and ARKit don't have these
-#if !(PLATFORM_HOLOLENS || PLATFORM_IOS)
 	/** Texture coordinates buffer */
 	FMRMeshVertexBuffer<FVector2D> UVBuffer;
 	/** Tangent space buffer */
 	FMRMeshVertexBuffer<FPackedNormal> TangentXZBuffer;
 	/** We don't need color */
 	FMRMeshVertexBuffer<FColor> ColorBuffer;
-#endif
 	/** Index buffer for this section */
 	FMRMeshIndexBuffer IndexBuffer;
 	/** Vertex factory for this section */
@@ -134,12 +131,9 @@ struct FMRMeshProxySection
 	void ReleaseResources()
 	{
 		PositionBuffer.ReleaseResource();
-		// AR meshes from HoloLens and ARKit don't have these
-#if !(PLATFORM_HOLOLENS || PLATFORM_IOS)
 		UVBuffer.ReleaseResource();
 		TangentXZBuffer.ReleaseResource();
 		ColorBuffer.ReleaseResource();
-#endif
 		IndexBuffer.ReleaseResource();
 		VertexFactory.ReleaseResource();
 	}
@@ -163,8 +157,6 @@ static void InitVertexFactory(FLocalVertexFactory* VertexFactory, const FMRMeshP
 			NewData.PositionComponent = FVertexStreamComponent(&MRMeshSection.PositionBuffer, 0, sizeof(FVector), VET_Float3, EVertexStreamUsage::Default);
 		}
 
-		// AR meshes from HoloLens and ARKit don't have these
-#if !(PLATFORM_HOLOLENS || PLATFORM_IOS)
 		if (MRMeshSection.UVBuffer.NumVerts != 0)
 		{
 			NewData.TextureCoordinatesSRV = MRMeshSection.UVBufferSRV;
@@ -184,7 +176,6 @@ static void InitVertexFactory(FLocalVertexFactory* VertexFactory, const FMRMeshP
 			NewData.ColorComponentsSRV = MRMeshSection.ColorBufferSRV;
 			NewData.ColorComponent = FVertexStreamComponent(&MRMeshSection.ColorBuffer, 0, sizeof(FColor), VET_Color, EVertexStreamUsage::ManualFetch);
 		}
-#endif
 
 		VertexFactory->SetData(NewData);
 		VertexFactory->InitResource();
@@ -202,10 +193,18 @@ public:
 	}
 
 	FMRMeshProxy(const UMRMeshComponent* InComponent)
-	: FPrimitiveSceneProxy(InComponent, InComponent->GetFName())
-	, MaterialToUse((InComponent->Material != nullptr) ? InComponent->Material : (InComponent->WireframeMaterial) ? InComponent->WireframeMaterial : UMaterial::GetDefaultMaterial(MD_Surface))
-	, FeatureLevel(GetScene().GetFeatureLevel())
-	{
+		: FPrimitiveSceneProxy(InComponent, InComponent->GetFName())
+		, MaterialToUse(nullptr)
+		, FeatureLevel(GetScene().GetFeatureLevel())
+		, bEnableOcclusion(InComponent->GetEnableMeshOcclusion())
+		, bUseWireframeForNoMaterial(InComponent->GetUseWireframe())
+		, bShouldRender(false)
+	{
+		if (InComponent->Material == nullptr && bUseWireframeForNoMaterial)
+		{
+			MaterialToUse = InComponent->WireframeMaterial;
+		}
+		bShouldRender = MaterialToUse != nullptr || bEnableOcclusion;
 	}
 
 	virtual ~FMRMeshProxy()
@@ -227,14 +226,11 @@ public:
 		FMRMeshProxySection* NewSection = new FMRMeshProxySection(Args.BrickId, FeatureLevel);
 		ProxySections.Add(NewSection);
 
-		// AR meshes from HoloLens and ARKit don't have these
-#if !(PLATFORM_HOLOLENS || PLATFORM_IOS)
 		// Vulkan requires that all the buffers be full.
 		const int32 NumVerts = Args.PositionData.Num();
 		check((NumVerts == Args.ColorData.Num()));
 		check((NumVerts == Args.UVData.Num()));
 		check((NumVerts * 2) == Args.TangentXZData.Num());
-#endif
 
 		// POSITION BUFFER
 		{
@@ -243,8 +239,6 @@ public:
 			NewSection->PositionBufferSRV = RHICreateShaderResourceView(NewSection->PositionBuffer.VertexBufferRHI, sizeof(float), PF_R32_FLOAT);
 		}
 
-		// AR meshes from HoloLens and ARKit don't have these
-#if !(PLATFORM_HOLOLENS || PLATFORM_IOS)
 		// TEXTURE COORDS BUFFER
 		{
 			NewSection->UVBuffer.InitResource();
@@ -278,7 +272,6 @@ public:
 				NewSection->ColorBufferSRV = RHICreateShaderResourceView(NewSection->ColorBuffer.VertexBufferRHI, 4, PF_R8G8B8A8);
 			}
 		}
-#endif
 
 		// INDEX BUFFER
 		{
@@ -331,6 +324,11 @@ private:
 	{
 		static const FBoxSphereBounds InfiniteBounds(FSphere(FVector::ZeroVector, HALF_WORLD_MAX));
 
+		if (!bShouldRender)
+		{
+			return;
+		}
+
 		// Iterate over sections
 		for (const FMRMeshProxySection* Section : ProxySections)
 		{
@@ -349,7 +347,9 @@ private:
 						FMeshBatch& Mesh = Collector.AllocateMesh();
 						FMeshBatchElement& BatchElement = Mesh.Elements[0];
 						BatchElement.IndexBuffer = &Section->IndexBuffer;
-						Mesh.bWireframe = false;
+						Mesh.bWireframe = bUseWireframeForNoMaterial;
+						Mesh.bUseAsOccluder = bEnableOcclusion;
+						Mesh.bUseForDepthPass = bEnableOcclusion;
 						Mesh.VertexFactory = &Section->VertexFactory;
 						Mesh.MaterialRenderProxy = MaterialProxy;
 
@@ -378,7 +378,7 @@ private:
 		Result.bDrawRelevance = IsShown(View);
 		Result.bShadowRelevance = IsShadowCast(View);
 		Result.bDynamicRelevance = true;
-		Result.bRenderInMainPass = ShouldRenderInMainPass();
+		Result.bRenderInMainPass = bShouldRender && ShouldRenderInMainPass();
 		Result.bUsesLightingChannels = GetLightingChannelMask() != GetDefaultLightingChannelMask();
 		Result.bRenderCustomDepth = ShouldRenderCustomDepth();
 		Result.bSeparateTranslucencyRelevance = MaterialToUse->GetMaterial()->bEnableSeparateTranslucency;
@@ -399,11 +399,16 @@ private:
 	TArray<FMRMeshProxySection*> ProxySections;
 	UMaterialInterface* MaterialToUse;
 	ERHIFeatureLevel::Type FeatureLevel;
+	bool bEnableOcclusion;
+	bool bUseWireframeForNoMaterial;
+	bool bShouldRender;
 };
 
 
 UMRMeshComponent::UMRMeshComponent(const FObjectInitializer& ObjectInitializer)
-: Super(ObjectInitializer)
+	: Super(ObjectInitializer)
+	, bEnableOcclusion(false)
+	, bUseWireframeForNoMaterial(false)
 {
 	static ConstructorHelpers::FObjectFinder<UMaterialInterface> WireframeMaterialRef(TEXT("/Engine/EngineDebugMaterials/WireframeMaterial"));
 	WireframeMaterial = WireframeMaterialRef.Object;
@@ -867,9 +872,10 @@ struct FMeshArrayHolder :
 	}
 };
 
-void UMRMeshComponent::UpdateTransformAndMesh(const FTransform& InTransform, TArray<FVector>& Vertices, TArray<MRMESH_INDEX_TYPE>& Indices)
+void UMRMeshComponent::UpdateMesh(const FVector& InLocation, const FQuat& InRotation, const FVector& Scale, TArray<FVector>& Vertices, TArray<MRMESH_INDEX_TYPE>& Indices)
 {
-	SetComponentToWorld(InTransform);
+	SetRelativeLocationAndRotation(InLocation, InRotation);
+	SetRelativeScale3D(Scale);
 
 	// Create our struct that will hold the data until the render thread is done with it
 	TSharedPtr<FMeshArrayHolder, ESPMode::ThreadSafe> MeshHolder(new FMeshArrayHolder(Vertices, Indices));
diff --git a/Engine/Source/Runtime/MRMesh/Public/MRMeshComponent.h b/Engine/Source/Runtime/MRMesh/Public/MRMeshComponent.h
--- a/Engine/Source/Runtime/MRMesh/Public/MRMeshComponent.h
+++ b/Engine/Source/Runtime/MRMesh/Public/MRMeshComponent.h
@@ -88,7 +88,13 @@ public:
 	virtual void SetMaterial(int32 ElementIndex, class UMaterialInterface* InMaterial) override;
 
 	/** Updates from HoloLens or iOS */
-	void UpdateTransformAndMesh(const FTransform& InTransform, TArray<FVector>& Vertices, TArray<MRMESH_INDEX_TYPE>& Indices);
+	void UpdateMesh(const FVector& InLocation, const FQuat& InRotation, const FVector& Scale, TArray<FVector>& Vertices, TArray<MRMESH_INDEX_TYPE>& Indices);
+
+	void SetEnableMeshOcclusion(bool bEnable) { bEnableOcclusion = bEnable; }
+	bool GetEnableMeshOcclusion() const { return bEnableOcclusion; }
+	void SetUseWireframe(bool bUseWireframe) { bUseWireframeForNoMaterial = bUseWireframe; }
+	bool GetUseWireframe() const { return bUseWireframeForNoMaterial; }
+
 
 protected:
 	virtual void OnActorEnableCollisionChanged() override;
@@ -103,6 +109,9 @@ public:
 
 	virtual void SetWalkableSlopeOverride(const FWalkableSlopeOverride& NewOverride) override;
 
+	void SetNeverCreateCollisionMesh(bool bNeverCreate) { bNeverCreateCollisionMesh = bNeverCreate; }
+	void SetEnableNavMesh(bool bEnable) { bUpdateNavMeshOnMeshUpdate = bEnable;  }
+
 private:
 	//~ UPrimitiveComponent
 	virtual FPrimitiveSceneProxy* CreateSceneProxy() override;
@@ -157,6 +166,11 @@ private:
 	UPROPERTY()
 	UMaterialInterface* WireframeMaterial;
 
+	/** Whether this mesh should write z-depth to occlude meshes or not */
+	bool bEnableOcclusion;
+	/** Whether this mesh should draw using the wireframe material when no material is set or not */
+	bool bUseWireframeForNoMaterial;
+
 	TArray<FBodyInstance*> BodyInstances;
 	TArray<IMRMesh::FBrickId> BodyIds;
-};
\ No newline at end of file
+};
