Removed the HMD dependency for MRMesh caused by the mock object

diff --git a/Engine/Source/Runtime/MRMesh/MRMesh.build.cs b/Engine/Source/Runtime/MRMesh/MRMesh.build.cs
--- a/Engine/Source/Runtime/MRMesh/MRMesh.build.cs
+++ b/Engine/Source/Runtime/MRMesh/MRMesh.build.cs
@@ -17,7 +17,6 @@ namespace UnrealBuildTool.Rules
                     "Engine",
                     "RenderCore",
                     "RHI",
-					"HeadMountedDisplay"
 				}
 			);
 
diff --git a/Engine/Source/Runtime/MRMesh/Private/MockDataMeshTrackerComponent.cpp b/Engine/Source/Runtime/MRMesh/Private/MockDataMeshTrackerComponent.cpp
--- a/Engine/Source/Runtime/MRMesh/Private/MockDataMeshTrackerComponent.cpp
+++ b/Engine/Source/Runtime/MRMesh/Private/MockDataMeshTrackerComponent.cpp
@@ -439,7 +439,7 @@ void UMockDataMeshTrackerComponent::UpdateBlock(int32 BlockIndex)
 	CurrentMeshDataCache->BrickId = BrickId;
 
 	// Pull vertices
-	const FVector VertexOffset = UHeadMountedDisplayFunctionLibrary::GetTrackingToWorldTransform(this).Inverse().GetLocation();
+	const FVector VertexOffset = FVector::ZeroVector;
 	const int32 VertexCount = RawMeshData.Vertices.Num();
 	CurrentMeshDataCache->OffsetVertices.Reserve(VertexCount);
 	CurrentMeshDataCache->WorldVertices.Reserve(VertexCount);
