Reverted some of my MRMesh changes until Jules gets the occlusion checked in

diff --git a/Engine/Source/Runtime/MRMesh/Private/MRMeshComponent.cpp b/Engine/Source/Runtime/MRMesh/Private/MRMeshComponent.cpp
--- a/Engine/Source/Runtime/MRMesh/Private/MRMeshComponent.cpp
+++ b/Engine/Source/Runtime/MRMesh/Private/MRMeshComponent.cpp
@@ -197,14 +197,17 @@ public:
 		, MaterialToUse(nullptr)
 		, FeatureLevel(GetScene().GetFeatureLevel())
 		, bEnableOcclusion(InComponent->GetEnableMeshOcclusion())
-		, bUseWireframeForNoMaterial(InComponent->GetUseWireframe())
-		, bShouldRender(false)
+		, bUseWireframe(InComponent->GetUseWireframe())
 	{
-		if (InComponent->Material == nullptr && bUseWireframeForNoMaterial)
+		if (bUseWireframe)
 		{
 			MaterialToUse = InComponent->WireframeMaterial;
 		}
-		bShouldRender = MaterialToUse != nullptr || bEnableOcclusion;
+		// If this is still null, use the default material
+		if (InComponent->Material == nullptr)
+		{
+			MaterialToUse = UMaterial::GetDefaultMaterial(MD_Surface);
+		}
 	}
 
 	virtual ~FMRMeshProxy()
@@ -324,11 +327,6 @@ private:
 	{
 		static const FBoxSphereBounds InfiniteBounds(FSphere(FVector::ZeroVector, HALF_WORLD_MAX));
 
-		if (!bShouldRender)
-		{
-			return;
-		}
-
 		// Iterate over sections
 		for (const FMRMeshProxySection* Section : ProxySections)
 		{
@@ -347,7 +345,7 @@ private:
 						FMeshBatch& Mesh = Collector.AllocateMesh();
 						FMeshBatchElement& BatchElement = Mesh.Elements[0];
 						BatchElement.IndexBuffer = &Section->IndexBuffer;
-						Mesh.bWireframe = bUseWireframeForNoMaterial;
+						Mesh.bWireframe = bUseWireframe;
 						Mesh.bUseAsOccluder = bEnableOcclusion;
 						Mesh.bUseForDepthPass = bEnableOcclusion;
 						Mesh.VertexFactory = &Section->VertexFactory;
@@ -378,7 +376,7 @@ private:
 		Result.bDrawRelevance = IsShown(View);
 		Result.bShadowRelevance = IsShadowCast(View);
 		Result.bDynamicRelevance = true;
-		Result.bRenderInMainPass = bShouldRender && ShouldRenderInMainPass();
+		Result.bRenderInMainPass = ShouldRenderInMainPass();
 		Result.bUsesLightingChannels = GetLightingChannelMask() != GetDefaultLightingChannelMask();
 		Result.bRenderCustomDepth = ShouldRenderCustomDepth();
 		Result.bSeparateTranslucencyRelevance = MaterialToUse->GetMaterial()->bEnableSeparateTranslucency;
@@ -400,8 +398,7 @@ private:
 	UMaterialInterface* MaterialToUse;
 	ERHIFeatureLevel::Type FeatureLevel;
 	bool bEnableOcclusion;
-	bool bUseWireframeForNoMaterial;
-	bool bShouldRender;
+	bool bUseWireframe;
 };
 
 
