Changed the way the MaterialToUse was setting the component and default materials

diff --git a/Engine/Source/Runtime/MRMesh/Private/MRMeshComponent.cpp b/Engine/Source/Runtime/MRMesh/Private/MRMeshComponent.cpp
--- a/Engine/Source/Runtime/MRMesh/Private/MRMeshComponent.cpp
+++ b/Engine/Source/Runtime/MRMesh/Private/MRMeshComponent.cpp
@@ -194,7 +194,7 @@ public:
 
 	FMRMeshProxy(const UMRMeshComponent* InComponent)
 		: FPrimitiveSceneProxy(InComponent, InComponent->GetFName())
-		, MaterialToUse(nullptr)
+		, MaterialToUse(InComponent->Material)
 		, FeatureLevel(GetScene().GetFeatureLevel())
 		, bEnableOcclusion(InComponent->GetEnableMeshOcclusion())
 		, bUseWireframe(InComponent->GetUseWireframe())
@@ -204,7 +204,7 @@ public:
 			MaterialToUse = InComponent->WireframeMaterial;
 		}
 		// If this is still null, use the default material
-		if (InComponent->Material == nullptr)
+		if (MaterialToUse == nullptr)
 		{
 			MaterialToUse = UMaterial::GetDefaultMaterial(MD_Surface);
 		}
