Improved the data validation for a MRMesh update

diff --git a/Engine/Source/Runtime/MRMesh/Private/MRMeshComponent.cpp b/Engine/Source/Runtime/MRMesh/Private/MRMeshComponent.cpp
--- a/Engine/Source/Runtime/MRMesh/Private/MRMeshComponent.cpp
+++ b/Engine/Source/Runtime/MRMesh/Private/MRMeshComponent.cpp
@@ -640,6 +640,8 @@ void UMRMeshComponent::SendBrickData_Internal(IMRMesh::FSendBrickDataArgs Args)
 
 	UE_LOG(LogMrMesh, Log, TEXT("SendBrickData_Internal() processing brick %llu with %i triangles"), Args.BrickId, Args.Indices.Num() / 3);
 
+	const bool bHasBrickData = Args.Indices.Num() > 0 && Args.PositionData.Num() > 0;
+
 	if (!IsPendingKill() && !bNeverCreateCollisionMesh)
 	{
 		// Physics update
@@ -648,7 +650,7 @@ void UMRMeshComponent::SendBrickData_Internal(IMRMesh::FSendBrickDataArgs Args)
 		{
 			int32 BodyIndex = BodyIds.Find(Args.BrickId);
 
-			if (const bool bBrickHasData = Args.Indices.Num() > 0)
+			if (bHasBrickData)
 			{
 				bPhysicsStateCreated = true;
 
@@ -707,7 +709,7 @@ void UMRMeshComponent::SendBrickData_Internal(IMRMesh::FSendBrickDataArgs Args)
 				}
 			}
 		}
-		if (bUpdateNavMeshOnMeshUpdate && bHasCustomNavigableGeometry)
+		if (bHasBrickData && bUpdateNavMeshOnMeshUpdate && bHasCustomNavigableGeometry)
 		{
 			UpdateNavigationData();
 		}
@@ -724,14 +726,14 @@ void UMRMeshComponent::SendBrickData_Internal(IMRMesh::FSendBrickDataArgs Args)
 			// Graphics update
 			UMRMeshComponent* This = this;
 			ENQUEUE_RENDER_COMMAND(FSendBrickDataLambda)(
-				[This, Args](FRHICommandListImmediate& RHICmdList)
+				[This, Args, bHasBrickData](FRHICommandListImmediate& RHICmdList)
 				{
 					FMRMeshProxy* MRMeshProxy = static_cast<FMRMeshProxy*>(This->SceneProxy);
 					if (MRMeshProxy)
 					{
 						MRMeshProxy->RenderThread_RemoveSection(Args.BrickId);
 
-						if (const bool bBrickHasData = Args.Indices.Num() > 0)
+						if (bHasBrickData)
 						{
 							MRMeshProxy->RenderThread_UploadNewSection(Args);
 						}
