Prep for Jules' occlusion support

diff --git a/Engine/Source/Runtime/MRMesh/Private/MRMeshComponent.cpp b/Engine/Source/Runtime/MRMesh/Private/MRMeshComponent.cpp
--- a/Engine/Source/Runtime/MRMesh/Private/MRMeshComponent.cpp
+++ b/Engine/Source/Runtime/MRMesh/Private/MRMeshComponent.cpp
@@ -347,7 +347,10 @@ private:
 						BatchElement.IndexBuffer = &Section->IndexBuffer;
 						Mesh.bWireframe = bUseWireframe;
 						Mesh.bUseAsOccluder = bEnableOcclusion;
-						Mesh.bUseForDepthPass = bEnableOcclusion;
+
+// Waiting for Jules' changes
+//						Mesh.bUseForDepthPass = bEnableOcclusion;
+
 						Mesh.VertexFactory = &Section->VertexFactory;
 						Mesh.MaterialRenderProxy = MaterialProxy;
 
@@ -376,7 +379,12 @@ private:
 		Result.bDrawRelevance = IsShown(View);
 		Result.bShadowRelevance = IsShadowCast(View);
 		Result.bDynamicRelevance = true;
-		Result.bRenderInMainPass = ShouldRenderInMainPass();
+		// If there is a material set that is not the default material, then this wants to be rendered in the main pass
+		Result.bRenderInMainPass = (bUseWireframe || MaterialToUse != UMaterial::GetDefaultMaterial(MD_Surface)) && ShouldRenderInMainPass();
+
+// Waiting for Jules' changes
+//		Result.bRenderInDepthPass = bEnableOcclusion;
+
 		Result.bUsesLightingChannels = GetLightingChannelMask() != GetDefaultLightingChannelMask();
 		Result.bRenderCustomDepth = ShouldRenderCustomDepth();
 		Result.bSeparateTranslucencyRelevance = MaterialToUse->GetMaterial()->bEnableSeparateTranslucency;
