Update to MrMesh to add frustum culling to components using draw wire box

diff --git a/Engine/Source/Runtime/MRMesh/Private/MRMeshComponent.cpp b/Engine/Source/Runtime/MRMesh/Private/MRMeshComponent.cpp
--- a/Engine/Source/Runtime/MRMesh/Private/MRMeshComponent.cpp
+++ b/Engine/Source/Runtime/MRMesh/Private/MRMeshComponent.cpp
@@ -26,6 +26,22 @@
 
 DECLARE_CYCLE_STAT(TEXT("MrMesh SetCollisionProfileName"), STAT_MrMesh_SetCollisionProfileName, STATGROUP_Physics);
 
+#define DEBUG_BRICK_CULLING
+
+#ifdef DEBUG_BRICK_CULLING
+enum class ECullingDebugState
+{
+	Off,
+	On,
+	Paused,
+};
+
+static TAutoConsoleVariable<int32> CVarPauseMRMeshBrickCulling(
+	TEXT("r.MrMesh.BrickCullingDebugState"),
+	static_cast<int32>(ECullingDebugState::Off),
+	TEXT("MR Mesh brick culling debug state: 0=off, 1=on, 2=paused"));
+#endif //DEBUG_BRICK_CULLING
+
 class FMRMeshVertexResourceArray : public FResourceArrayInterface
 {
 public:
@@ -116,6 +132,8 @@ struct FMRMeshProxySection
 	FMRMeshIndexBuffer IndexBuffer;
 	/** Vertex factory for this section */
 	FLocalVertexFactory VertexFactory;
+	/** AABB for this section */
+	FBox Bounds;
 
 	FShaderResourceViewRHIRef PositionBufferSRV;
 	FShaderResourceViewRHIRef UVBufferSRV;
@@ -286,6 +304,9 @@ public:
 		{
 			InitVertexFactory(&NewSection->VertexFactory, *NewSection);
 		}
+
+		// BOUNDS
+		NewSection->Bounds = Args.Bounds;
 	}
 
 	bool RenderThread_RemoveSection(IMRMesh::FBrickId BrickId)
@@ -327,12 +348,17 @@ private:
 	{
 		static const FBoxSphereBounds InfiniteBounds(FSphere(FVector::ZeroVector, HALF_WORLD_MAX));
 
+#ifdef DEBUG_BRICK_CULLING
+		TMap<IMRMesh::FBrickId, TTuple<FBox, bool>> NewVisDataByBrickId;
+		ECullingDebugState CullingDebugState =
+			static_cast<ECullingDebugState>(CVarPauseMRMeshBrickCulling.GetValueOnRenderThread());
+#endif //DEBUG_BRICK_CULLING
+
 		// Iterate over sections
 		for (const FMRMeshProxySection* Section : ProxySections)
 		{
 			if (Section != nullptr)
 			{
-				const bool bIsSelected = false;
 				FMaterialRenderProxy* MaterialProxy = MaterialToUse->GetRenderProxy();
 
 				// For each view..
@@ -341,36 +367,90 @@ private:
 					if (VisibilityMap & (1 << ViewIndex))
 					{
 						const FSceneView* View = Views[ViewIndex];
-						// Draw the mesh.
-						FMeshBatch& Mesh = Collector.AllocateMesh();
-						FMeshBatchElement& BatchElement = Mesh.Elements[0];
-						BatchElement.IndexBuffer = &Section->IndexBuffer;
-						Mesh.bWireframe = bUseWireframe;
-						Mesh.bUseAsOccluder = bEnableOcclusion;
+						bool IsVisible = Section->Bounds.GetExtent().IsNearlyZero() || View->ViewFrustum.
+							IntersectBox(Section->Bounds.GetCenter(), Section->Bounds.GetExtent());
+
+#ifdef DEBUG_BRICK_CULLING
+						switch (CullingDebugState)
+						{
+							case ECullingDebugState::Off:
+								break;
+							case ECullingDebugState::On:
+								NewVisDataByBrickId.Add(Section->BrickId, MakeTuple(Section->Bounds, IsVisible));
+								break;
+							case ECullingDebugState::Paused:
+								auto OldVisData = OldVisDataByBrickId.Find(Section->BrickId);
+								if (OldVisData)
+								{
+									NewVisDataByBrickId.Add(Section->BrickId, *OldVisData);
+
+									// Easier to see what's culled if mesh mimics culling pause state
+									IsVisible = OldVisData->Value;
+								}
+								else
+								{
+									IsVisible = false;
+								}
+								break;
+						}
+#endif //DEBUG_BRICK_CULLING
+
+						if (IsVisible)
+						{
+							// Draw the mesh.
+							FMeshBatch& Mesh = Collector.AllocateMesh();
+							FMeshBatchElement& BatchElement = Mesh.Elements[0];
+							BatchElement.IndexBuffer = &Section->IndexBuffer;
+							Mesh.bWireframe = bUseWireframe;
+							Mesh.bUseAsOccluder = bEnableOcclusion;
 
 // Waiting for Jules' changes
-//						Mesh.bUseForDepthPass = bEnableOcclusion;
-
-						Mesh.VertexFactory = &Section->VertexFactory;
-						Mesh.MaterialRenderProxy = MaterialProxy;
-
-						FDynamicPrimitiveUniformBuffer& DynamicPrimitiveUniformBuffer = Collector.AllocateOneFrameResource<FDynamicPrimitiveUniformBuffer>();
-						DynamicPrimitiveUniformBuffer.Set(GetLocalToWorld(), GetLocalToWorld(), InfiniteBounds, InfiniteBounds, true, false, UseEditorDepthTest(), false);
-						BatchElement.PrimitiveUniformBufferResource = &DynamicPrimitiveUniformBuffer.UniformBuffer;
-
-						BatchElement.FirstIndex = 0;
-						BatchElement.NumPrimitives = Section->IndexBuffer.NumIndices / 3;
-						BatchElement.MinVertexIndex = 0;
-						BatchElement.MaxVertexIndex = Section->PositionBuffer.NumVerts - 1;
-						Mesh.ReverseCulling = IsLocalToWorldDeterminantNegative();
-						Mesh.Type = PT_TriangleList;
-						Mesh.DepthPriorityGroup = SDPG_World;
-						Mesh.bCanApplyViewModeOverrides = false;
-						Collector.AddMesh(ViewIndex, Mesh);
+//							Mesh.bUseForDepthPass = bEnableOcclusion;
+
+							Mesh.VertexFactory = &Section->VertexFactory;
+							Mesh.MaterialRenderProxy = MaterialProxy;
+
+							FDynamicPrimitiveUniformBuffer& DynamicPrimitiveUniformBuffer = Collector.AllocateOneFrameResource<FDynamicPrimitiveUniformBuffer>();
+							DynamicPrimitiveUniformBuffer.Set(GetLocalToWorld(), GetLocalToWorld(), InfiniteBounds, InfiniteBounds, true, false, UseEditorDepthTest(), false);
+							BatchElement.PrimitiveUniformBufferResource = &DynamicPrimitiveUniformBuffer.UniformBuffer;
+
+							BatchElement.FirstIndex = 0;
+							BatchElement.NumPrimitives = Section->IndexBuffer.NumIndices / 3;
+							BatchElement.MinVertexIndex = 0;
+							BatchElement.MaxVertexIndex = Section->PositionBuffer.NumVerts - 1;
+							Mesh.ReverseCulling = IsLocalToWorldDeterminantNegative();
+							Mesh.Type = PT_TriangleList;
+							Mesh.DepthPriorityGroup = SDPG_World;
+							Mesh.bCanApplyViewModeOverrides = false;
+							Collector.AddMesh(ViewIndex, Mesh);
+						}
+					}
+				}
+			}
+		}
+
+#ifdef DEBUG_BRICK_CULLING
+		OldVisDataByBrickId = NewVisDataByBrickId;
+
+		if (CullingDebugState != ECullingDebugState::Off)
+		{
+			const FColor ColorGray(0x7f, 0x7f, 0x7f, 0xff);
+
+			for (int32 ViewIndex = 0; ViewIndex < Views.Num(); ViewIndex++)
+			{
+				if (VisibilityMap & (1 << ViewIndex))
+				{
+					for (const auto& BrickVisibility : OldVisDataByBrickId)
+					{
+						const auto& BrickBounds = BrickVisibility.Value.Key;
+						FColor BoundsColor(BrickVisibility.Value.Value ? FColor::Green : ColorGray);
+						FPrimitiveDrawInterface *PDI = Collector.GetPDI(ViewIndex);
+						DrawWireBox(PDI, BrickBounds, BoundsColor, GetDepthPriorityGroup(Views[ViewIndex]));
 					}
 				}
 			}
 		}
+#endif //DEBUG_BRICK_CULLING
 	}
 
 	virtual FPrimitiveViewRelevance GetViewRelevance(const FSceneView* View) const
@@ -407,6 +487,10 @@ private:
 	ERHIFeatureLevel::Type FeatureLevel;
 	bool bEnableOcclusion;
 	bool bUseWireframe;
+
+#ifdef DEBUG_BRICK_CULLING
+	mutable TMap<IMRMesh::FBrickId, TTuple<FBox, bool>> OldVisDataByBrickId;
+#endif //DEBUG_BRICK_CULLING
 };
 
 
@@ -793,6 +877,11 @@ void UMRMeshComponent::ClearAllBrickData_Internal()
 			}
 		}
 	);
+
+	if (OnClear().IsBound())
+	{
+		OnClear().Broadcast();
+	}
 }
 
 void UMRMeshComponent::SetMaterial(int32 ElementIndex, class UMaterialInterface* InMaterial)
diff --git a/Engine/Source/Runtime/MRMesh/Public/MRMeshComponent.h b/Engine/Source/Runtime/MRMesh/Public/MRMeshComponent.h
--- a/Engine/Source/Runtime/MRMesh/Public/MRMeshComponent.h
+++ b/Engine/Source/Runtime/MRMesh/Public/MRMeshComponent.h
@@ -48,6 +48,7 @@ public:
 		const TArray<FPackedNormal>& TangentXZData;
 		const TArray<FColor>& ColorData;
 		const TArray<MRMESH_INDEX_TYPE>& Indices;
+		const FBox Bounds;
 	};
 
 	virtual void SetConnected(bool value) = 0;
@@ -112,6 +113,10 @@ public:
 	void SetNeverCreateCollisionMesh(bool bNeverCreate) { bNeverCreateCollisionMesh = bNeverCreate; }
 	void SetEnableNavMesh(bool bEnable) { bUpdateNavMeshOnMeshUpdate = bEnable;  }
 
+	/** Trackers feeding mesh data to this component may want to know when we clear our mesh data */
+	DECLARE_EVENT(UMRMeshComponent, FOnClear);
+	FOnClear& OnClear() { return OnClearEvent; }
+
 private:
 	//~ UPrimitiveComponent
 	virtual FPrimitiveSceneProxy* CreateSceneProxy() override;
@@ -173,4 +178,6 @@ private:
 
 	TArray<FBodyInstance*> BodyInstances;
 	TArray<IMRMesh::FBrickId> BodyIds;
-};
+
+	FOnClear OnClearEvent;
+};
\ No newline at end of file
