VicunaSolo fixes Related to UEVR-1505 Hololens 2 fixes from microsoft for 4.24
-Fix for MRMesh GetMaterial from a worker thread triggering a check..
-Add tracking status update to anchors so we can tell if the device has lost tracking on one.

diff --git a/Engine/Plugins/Runtime/AR/Microsoft/HoloLensAR/Source/HoloLensAR/Private/HoloLensARSystem.cpp b/Engine/Plugins/Runtime/AR/Microsoft/HoloLensAR/Source/HoloLensAR/Private/HoloLensARSystem.cpp
--- a/Engine/Plugins/Runtime/AR/Microsoft/HoloLensAR/Source/HoloLensAR/Private/HoloLensARSystem.cpp
+++ b/Engine/Plugins/Runtime/AR/Microsoft/HoloLensAR/Source/HoloLensAR/Private/HoloLensARSystem.cpp
@@ -499,6 +499,11 @@ void FHoloLensARSystem::UpdateWMRAnchors()
 			if (WMRGetAnchorTransform(*AnchorId, Transform))
 			{
 				Pin->OnTransformUpdated(Transform);
+				Pin->OnTrackingStateChanged(EARTrackingState::Tracking);
+			}
+			else
+			{
+				Pin->OnTrackingStateChanged(EARTrackingState::NotTracking);
 			}
 		}
 	}
diff --git a/Engine/Source/Runtime/MRMesh/Private/MRMeshComponent.cpp b/Engine/Source/Runtime/MRMesh/Private/MRMeshComponent.cpp
--- a/Engine/Source/Runtime/MRMesh/Private/MRMeshComponent.cpp
+++ b/Engine/Source/Runtime/MRMesh/Private/MRMeshComponent.cpp
@@ -468,7 +468,8 @@ private:
 		Result.bRenderInDepthPass = bEnableOcclusion;
 		Result.bUsesLightingChannels = GetLightingChannelMask() != GetDefaultLightingChannelMask();
 		Result.bRenderCustomDepth = ShouldRenderCustomDepth();
-		Result.bSeparateTranslucencyRelevance = MaterialToUse->GetMaterial()->bEnableSeparateTranslucency;
+		UMaterialInterface::TMicRecursionGuard RecursionGuard;
+		Result.bSeparateTranslucencyRelevance = MaterialToUse->GetMaterial_Concurrent(RecursionGuard)->bEnableSeparateTranslucency;
 		//MaterialRelevance.SetPrimitiveViewRelevance(Result);
 		return Result;
 	}
