Static analysis fixes.

diff --git a/Engine/Source/Runtime/MRMesh/Private/MRMeshComponent.cpp b/Engine/Source/Runtime/MRMesh/Private/MRMeshComponent.cpp
--- a/Engine/Source/Runtime/MRMesh/Private/MRMeshComponent.cpp
+++ b/Engine/Source/Runtime/MRMesh/Private/MRMeshComponent.cpp
@@ -610,7 +610,7 @@ void UMRMeshComponent::SendBrickData_Internal(IMRMesh::FSendBrickDataArgs Args)
 		{
 			int32 BodyIndex = BodyIds.Find(Args.BrickId);
 
-			if (const bool bBrickHasData = Args.Indices.Num() > 0)
+			if (/*const bool bBrickHasData = */Args.Indices.Num() > 0)
 			{
 				bPhysicsStateCreated = true;
 
@@ -693,7 +693,7 @@ void UMRMeshComponent::SendBrickData_Internal(IMRMesh::FSendBrickDataArgs Args)
 					{
 						MRMeshProxy->RenderThread_RemoveSection(Args.BrickId);
 
-						if (const bool bBrickHasData = Args.Indices.Num() > 0)
+						if (/*const bool bBrickHasData = */Args.Indices.Num() > 0)
 						{
 							MRMeshProxy->RenderThread_UploadNewSection(Args);
 						}
diff --git a/Engine/Source/Runtime/SlateCore/Private/Rendering/ElementBatcher.cpp b/Engine/Source/Runtime/SlateCore/Private/Rendering/ElementBatcher.cpp
--- a/Engine/Source/Runtime/SlateCore/Private/Rendering/ElementBatcher.cpp
+++ b/Engine/Source/Runtime/SlateCore/Private/Rendering/ElementBatcher.cpp
@@ -1594,7 +1594,7 @@ void FSlateElementBatcher::AddSplineElement(const FSlateDrawElement& DrawElement
 		SplineColor
 	);
 
-	if (const bool bNoGradient = InPayload.GradientStops.Num() <= 1)
+	if (/*const bool bNoGradient = */InPayload.GradientStops.Num() <= 1)
 	{
 		// Normal scenario where there is no color gradient.
 		LineBuilder.BuildBezierGeometry(InPayload.P0, InPayload.P1, InPayload.P2, InPayload.P3);
