Merge from Partner-MagicLeap-Staging: Workaround fix MRMesh flickering on Vk desktop renderer. The true cause of this issue is ultimately somewhere deeper.
To fix this, at least for now, enable mesh occlusion when the MRMesh is connected to the MagicLeapTrackerComponent. Since this happens after the MRMeshProxy is created, modify SetEnableMeshOcclusion so that is updates the proxy's value as well.

diff --git a/Engine/Plugins/Lumin/MagicLeap/Source/MagicLeap/Private/MagicLeapMeshTrackerComponent.cpp b/Engine/Plugins/Lumin/MagicLeap/Source/MagicLeap/Private/MagicLeapMeshTrackerComponent.cpp
--- a/Engine/Plugins/Lumin/MagicLeap/Source/MagicLeap/Private/MagicLeapMeshTrackerComponent.cpp
+++ b/Engine/Plugins/Lumin/MagicLeap/Source/MagicLeap/Private/MagicLeapMeshTrackerComponent.cpp
@@ -460,6 +460,8 @@ void UMagicLeapMeshTrackerComponent::ConnectMRMesh(UMRMeshComponent* InMRMeshPtr
 	{
 		InMRMeshPtr->SetConnected(true);
 		MRMesh = InMRMeshPtr;
+		// Enabling occlusion on the MRMesh prevents flickering due to an unknown UE4 issue.
+		MRMesh->SetEnableMeshOcclusion(true);
 		Impl->OnClearDelegateHandle = MRMesh->OnClear().AddRaw(Impl, &FMagicLeapMeshTrackerImpl::OnClear);
 	}
 }
diff --git a/Engine/Source/Runtime/MRMesh/Private/MRMeshComponent.cpp b/Engine/Source/Runtime/MRMesh/Private/MRMeshComponent.cpp
--- a/Engine/Source/Runtime/MRMesh/Private/MRMeshComponent.cpp
+++ b/Engine/Source/Runtime/MRMesh/Private/MRMeshComponent.cpp
@@ -347,6 +347,11 @@ public:
 #endif
 	}
 
+	void SetEnableMeshOcclusion(bool bEnable)
+	{
+		bEnableOcclusion = bEnable;
+	}
+
 private:
 	//~ FPrimitiveSceneProxy
 
@@ -991,3 +996,23 @@ void UMRMeshComponent::UpdateMesh(const FVector& InLocation, const FQuat& InRota
 	);
 }
 
+void UMRMeshComponent::SetEnableMeshOcclusion(bool bEnable)
+{
+	bEnableOcclusion = bEnable;
+
+	// Update bEnableOcclusion on the SceneProxy, as well.
+	if (SceneProxy)
+	{
+		UMRMeshComponent* This = this;
+		ENQUEUE_RENDER_COMMAND(FSetEnableMeshOcclusionLambda)(
+			[This,bEnable](FRHICommandListImmediate& RHICmdList)
+			{
+				FMRMeshProxy* MRMeshProxy = static_cast<FMRMeshProxy*>(This->SceneProxy);
+				if (MRMeshProxy)
+				{
+					MRMeshProxy->SetEnableMeshOcclusion(bEnable);
+				}
+			}
+		);
+	}
+}
diff --git a/Engine/Source/Runtime/MRMesh/Public/MRMeshComponent.h b/Engine/Source/Runtime/MRMesh/Public/MRMeshComponent.h
--- a/Engine/Source/Runtime/MRMesh/Public/MRMeshComponent.h
+++ b/Engine/Source/Runtime/MRMesh/Public/MRMeshComponent.h
@@ -91,7 +91,7 @@ public:
 	/** Updates from HoloLens or iOS */
 	void UpdateMesh(const FVector& InLocation, const FQuat& InRotation, const FVector& Scale, TArray<FVector>& Vertices, TArray<MRMESH_INDEX_TYPE>& Indices);
 
-	void SetEnableMeshOcclusion(bool bEnable) { bEnableOcclusion = bEnable; }
+	void SetEnableMeshOcclusion(bool bEnable);
 	bool GetEnableMeshOcclusion() const { return bEnableOcclusion; }
 	void SetUseWireframe(bool bUseWireframe) { bUseWireframeForNoMaterial = bUseWireframe; }
 	bool GetUseWireframe() const { return bUseWireframeForNoMaterial; }
