PR #8711: Fix MRMesh vertices (Contributed by [REDACTED_PERSON])

diff --git a/Engine/Plugins/Runtime/AR/AppleAR/AppleARKit/Source/AppleARKit/Public/AppleARKitMeshData.h b/Engine/Plugins/Runtime/AR/AppleAR/AppleARKit/Source/AppleARKit/Public/AppleARKitMeshData.h
--- a/Engine/Plugins/Runtime/AR/AppleAR/AppleARKit/Source/AppleARKit/Public/AppleARKitMeshData.h
+++ b/Engine/Plugins/Runtime/AR/AppleAR/AppleARKit/Source/AppleARKit/Public/AppleARKitMeshData.h
@@ -56,7 +56,7 @@ private:
 	
 private:
 	/** Mesh vertices in the anchor's local space */
-	TArray<FVector> Vertices;
+	TArray<FVector3f> Vertices;
 	
 	/** Mesh tangents in the anchor's local space */
 	TArray<FPackedNormal> TangentData;
diff --git a/Engine/Plugins/Runtime/MeshReconstruction/DummyMeshReconstructor/Source/DummyMeshReconstructor/Private/DummyMeshReconstructor.cpp b/Engine/Plugins/Runtime/MeshReconstruction/DummyMeshReconstructor/Source/DummyMeshReconstructor/Private/DummyMeshReconstructor.cpp
--- a/Engine/Plugins/Runtime/MeshReconstruction/DummyMeshReconstructor/Source/DummyMeshReconstructor/Private/DummyMeshReconstructor.cpp
+++ b/Engine/Plugins/Runtime/MeshReconstruction/DummyMeshReconstructor/Source/DummyMeshReconstructor/Private/DummyMeshReconstructor.cpp
@@ -141,7 +141,7 @@ private:
 	{
 		IMRMesh::FBrickId BrickId;
 		FIntVector BrickCoords;
-		TArray<FVector> PositionData;
+		TArray<FVector3f> PositionData;
 		TArray<FVector2D> UVData;
 		TArray<FPackedNormal> Tangents;
 		TArray<FColor> ColorData;
diff --git a/Engine/Source/Runtime/MRMesh/Private/MRMeshComponent.cpp b/Engine/Source/Runtime/MRMesh/Private/MRMeshComponent.cpp
--- a/Engine/Source/Runtime/MRMesh/Private/MRMeshComponent.cpp
+++ b/Engine/Source/Runtime/MRMesh/Private/MRMeshComponent.cpp
@@ -140,7 +140,7 @@ struct FMRMeshProxySection
 	/** Which brick this section represents */
 	IMRMesh::FBrickId BrickId;
 	/** Position buffer */
-	FMRMeshVertexBuffer<FVector> PositionBuffer;
+	FMRMeshVertexBuffer<FVector3f> PositionBuffer;
 	/** Texture coordinates buffer */
 	FMRMeshVertexBuffer<FVector2D> UVBuffer;
 	/** Tangent space buffer */
@@ -1013,7 +1013,7 @@ void UMRMeshComponent::Clear()
 struct FMeshArrayHolder :
 	public IMRMesh::FBrickDataReceipt
 {
-	TArray<FVector> Vertices;
+	TArray<FVector3f> Vertices;
 	TArray<MRMESH_INDEX_TYPE> Indices;
 	// Super wasteful of memory and perf, but the vertex factory requires these to be filled
 	// @todo Write a vertex factory that doesn't need all this overhead
diff --git a/Engine/Source/Runtime/MRMesh/Private/MockDataMeshTrackerComponent.cpp b/Engine/Source/Runtime/MRMesh/Private/MockDataMeshTrackerComponent.cpp
--- a/Engine/Source/Runtime/MRMesh/Private/MockDataMeshTrackerComponent.cpp
+++ b/Engine/Source/Runtime/MRMesh/Private/MockDataMeshTrackerComponent.cpp
@@ -127,7 +127,7 @@ public:
 		
 		IMRMesh::FBrickId BrickId = 0;
 		TArray<FVector> OffsetVertices;
-		TArray<FVector> WorldVertices;
+		TArray<FVector3f> WorldVertices;
 		TArray<FVector> Normals;
 		TArray<FVector2D> UV0;
 		TArray<FColor> VertexColors;
@@ -392,7 +392,7 @@ void UMockDataMeshTrackerComponent::RemoveBlock(int32 BlockIndex)
 	// Delete the brick and its cache entry
 	if (Impl->MeshBrickCache.Contains(BlockIndex))
 	{
-		const static TArray<FVector> EmptyVertices;
+		const static TArray<FVector3f> EmptyVertices;
 		const static TArray<FVector2D> EmptyUVs;
 		const static TArray<FPackedNormal> EmptyTangents;
 		const static TArray<FColor> EmptyVertexColors;
diff --git a/Engine/Source/Runtime/MRMesh/Public/MRMeshComponent.h b/Engine/Source/Runtime/MRMesh/Public/MRMeshComponent.h
--- a/Engine/Source/Runtime/MRMesh/Public/MRMeshComponent.h
+++ b/Engine/Source/Runtime/MRMesh/Public/MRMeshComponent.h
@@ -38,7 +38,7 @@ public:
 	{
 		TSharedPtr<FBrickDataReceipt, ESPMode::ThreadSafe> BrickDataReceipt;
 		const FBrickId BrickId;
-		const TArray<FVector>& PositionData;
+		const TArray<FVector3f>& PositionData;
 		const TArray<FVector2D>& UVData;
 		const TArray<FPackedNormal>& TangentXZData;
 		const TArray<FColor>& ColorData;
