Fix build break in UE5 Main (MRMesh vector type)

diff --git a/Engine/Source/Runtime/MRMesh/Private/MRMeshComponent.cpp b/Engine/Source/Runtime/MRMesh/Private/MRMeshComponent.cpp
--- a/Engine/Source/Runtime/MRMesh/Private/MRMeshComponent.cpp
+++ b/Engine/Source/Runtime/MRMesh/Private/MRMeshComponent.cpp
@@ -1146,11 +1146,7 @@ bool UMRMeshComponent::GetPhysicsTriMeshData(struct FTriMeshCollisionData* Colli
 	if (TempPosition && TempIndices)
 	{
 		// Copy the vertices
-		CollisionData->Vertices.Reset(TempPosition->Num());
-		for (auto Position = 0; Position < TempPosition->Num(); ++Position)
-		{
-			CollisionData->Vertices[Position] = (*TempPosition)[Position];
-		}
+		CollisionData->Vertices = *TempPosition;
 		
 		// Copy the indices
 		const auto& Indices = *TempIndices;
diff --git a/Engine/Source/Runtime/MRMesh/Public/MRMeshComponent.h b/Engine/Source/Runtime/MRMesh/Public/MRMeshComponent.h
--- a/Engine/Source/Runtime/MRMesh/Public/MRMeshComponent.h
+++ b/Engine/Source/Runtime/MRMesh/Public/MRMeshComponent.h
@@ -207,7 +207,7 @@ private:
 	FLinearColor WireframeColor = FLinearColor::White;
 	
 	/** Temporarily saved data pointers used for constructing the collision mesh */
-	const TArray<FVector>* TempPosition = nullptr;
+	const TArray<FVector3f>* TempPosition = nullptr;
 	const TArray<MRMESH_INDEX_TYPE>* TempIndices = nullptr;
 	
 	FOnMRMeshBrickDataUpdatedDelegate OnBrickDataUpdatedDelegate;
