UE-91500 Hololens generated mesh work once PIE with TP_HolloensBP Template
-Fixing async cooking of collision with chaos.  MockDataMeshTrackerComponent generated geometry can now have functional collision.
-Also improved test level and MockDataMeshTrackerComponent to help test this.
-Added method to cause a navmesh update only if collision has changed.  Fixed a couple cases where navigation would not have generated after a collision change.

diff --git a/Engine/Source/Runtime/MRMesh/Private/MRMeshComponent.cpp b/Engine/Source/Runtime/MRMesh/Private/MRMeshComponent.cpp
--- a/Engine/Source/Runtime/MRMesh/Private/MRMeshComponent.cpp
+++ b/Engine/Source/Runtime/MRMesh/Private/MRMeshComponent.cpp
@@ -550,9 +550,9 @@ void UMRMeshComponent::EndPlay(const EEndPlayReason::Type EndPlayReason)
 
 void UMRMeshComponent::OnActorEnableCollisionChanged()
 {
-	for (auto& BodyInstanceElement : BodyInstances)
+	for (auto& BodyHolderElement : BodyHolders)
 	{
-		BodyInstanceElement->UpdatePhysicsFilterData();
+		BodyHolderElement->BodyInstance.UpdatePhysicsFilterData();
 	}
 
 	Super::OnActorEnableCollisionChanged();
@@ -568,16 +568,16 @@ void UMRMeshComponent::SetCollisionEnabled(ECollisionEnabled::Type NewType)
 {
 	if (BodyInstance.GetCollisionEnabled() != NewType)
 	{
-		for (auto& BodyInstanceElement : BodyInstances)
+		for (auto& BodyHolderElement : BodyHolders)
 		{
-			BodyInstanceElement->SetCollisionEnabled(NewType);
+			BodyHolderElement->BodyInstance.SetCollisionEnabled(NewType);
 		}
 
 		if (IsRegistered() && BodyInstance.bSimulatePhysics && !IsWelded())
 		{
-			for (auto& BodyInstanceElement : BodyInstances)
+			for (auto& BodyHolderElement : BodyHolders)
 			{
-				BodyInstanceElement->ApplyWeldOnChildren();
+				BodyHolderElement->BodyInstance.ApplyWeldOnChildren();
 			}
 		}
 	}
@@ -593,16 +593,16 @@ void UMRMeshComponent::SetCollisionProfileName(FName InCollisionProfileName, boo
 	if (ThreadContext.ConstructedObject == this)
 	{
 		// If we are in our constructor, defer setup until PostInitProperties as derived classes
-		for (auto& BodyInstanceElement : BodyInstances)
+		for (auto& BodyHolderElement : BodyHolders)
 		{
-			BodyInstanceElement->SetCollisionProfileNameDeferred(InCollisionProfileName);
+			BodyHolderElement->BodyInstance.SetCollisionProfileNameDeferred(InCollisionProfileName);
 		}
 	}
 	else
 	{
-		for (auto& BodyInstanceElement : BodyInstances)
+		for (auto& BodyHolderElement : BodyHolders)
 		{
-			BodyInstanceElement->SetCollisionProfileName(InCollisionProfileName);
+			BodyHolderElement->BodyInstance.SetCollisionProfileName(InCollisionProfileName);
 		}
 	}
 
@@ -611,9 +611,9 @@ void UMRMeshComponent::SetCollisionProfileName(FName InCollisionProfileName, boo
 
 void UMRMeshComponent::SetCollisionObjectType(ECollisionChannel Channel)
 {
-	for (auto& BodyInstanceElement : BodyInstances)
+	for (auto& BodyHolderElement : BodyHolders)
 	{
-		BodyInstanceElement->SetObjectType(Channel);
+		BodyHolderElement->BodyInstance.SetObjectType(Channel);
 	}
 
 	Super::SetCollisionObjectType(Channel);
@@ -621,9 +621,9 @@ void UMRMeshComponent::SetCollisionObjectType(ECollisionChannel Channel)
 
 void UMRMeshComponent::SetCollisionResponseToChannel(ECollisionChannel Channel, ECollisionResponse NewResponse)
 {
-	for (auto& BodyInstanceElement : BodyInstances)
+	for (auto& BodyHolderElement : BodyHolders)
 	{
-		BodyInstanceElement->SetResponseToChannel(Channel, NewResponse);
+		BodyHolderElement->BodyInstance.SetResponseToChannel(Channel, NewResponse);
 	}
 
 	Super::SetCollisionResponseToChannel(Channel, NewResponse);
@@ -631,9 +631,9 @@ void UMRMeshComponent::SetCollisionResponseToChannel(ECollisionChannel Channel,
 
 void UMRMeshComponent::SetCollisionResponseToAllChannels(enum ECollisionResponse NewResponse)
 {
-	for (auto& BodyInstanceElement : BodyInstances)
+	for (auto& BodyHolderElement : BodyHolders)
 	{
-		BodyInstanceElement->SetResponseToAllChannels(NewResponse);
+		BodyHolderElement->BodyInstance.SetResponseToAllChannels(NewResponse);
 	}
 
 	Super::SetCollisionResponseToAllChannels(NewResponse);
@@ -641,9 +641,9 @@ void UMRMeshComponent::SetCollisionResponseToAllChannels(enum ECollisionResponse
 
 void UMRMeshComponent::SetCollisionResponseToChannels(const FCollisionResponseContainer& NewResponses)
 {
-	for (auto& BodyInstanceElement : BodyInstances)
+	for (auto& BodyHolderElement : BodyHolders)
 	{
-		BodyInstanceElement->SetResponseToChannels(NewResponses);
+		BodyHolderElement->BodyInstance.SetResponseToChannels(NewResponses);
 	}
 
 	Super::SetCollisionResponseToChannels(NewResponses);
@@ -651,11 +651,11 @@ void UMRMeshComponent::SetCollisionResponseToChannels(const FCollisionResponseCo
 
 void UMRMeshComponent::UpdatePhysicsToRBChannels()
 {
-	for (auto& BodyInstanceElement : BodyInstances)
+	for (auto& BodyHolderElement : BodyHolders)
 	{
-		if (BodyInstanceElement->IsValidBodyInstance())
+		if (BodyHolderElement->BodyInstance.IsValidBodyInstance())
 		{
-			BodyInstanceElement->UpdatePhysicsFilterData();
+			BodyHolderElement->BodyInstance.UpdatePhysicsFilterData();
 		}
 	}
 
@@ -664,11 +664,11 @@ void UMRMeshComponent::UpdatePhysicsToRBChannels()
 
 void UMRMeshComponent::SetWalkableSlopeOverride(const FWalkableSlopeOverride& NewOverride)
 {
-	for (auto& BodyInstanceElement : BodyInstances)
+	for (auto& BodyHolderElement : BodyHolders)
 	{
-		if (BodyInstanceElement->IsValidBodyInstance())
+		if (BodyHolderElement->BodyInstance.IsValidBodyInstance())
 		{
-			BodyInstanceElement->SetWalkableSlopeOverride(NewOverride);
+			BodyHolderElement->BodyInstance.SetWalkableSlopeOverride(NewOverride);
 		}
 	}
 
@@ -721,41 +721,6 @@ void UMRMeshComponent::ClearAllBrickData()
 	FSimpleDelegateGraphTask::CreateAndDispatchWhenReady(ClearBrickDataTask, GET_STATID(STAT_UMRMeshComponent_ClearAllBrickData), nullptr, ENamedThreads::GameThread);
 }
 
-void UMRMeshComponent::CacheBodySetupHelper()
-{
-	CachedBodySetup = NewObject<UBodySetup>(this, NAME_None);
-	CachedBodySetup->BodySetupGuid = FGuid::NewGuid();
-	CachedBodySetup->bGenerateMirroredCollision = false;
-	CachedBodySetup->bHasCookedCollisionData = true;
-}
-
-UBodySetup* UMRMeshComponent::CreateBodySetupHelper()
-{
-	// The body setup in a template needs to be public since the property is Instanced and thus is the archetype of the instance meaning there is a direct reference
-	UBodySetup* NewBS = NewObject<UBodySetup>(this, NAME_None);
-	NewBS->BodySetupGuid = FGuid::NewGuid();
-	NewBS->bGenerateMirroredCollision = false;
-	NewBS->bHasCookedCollisionData = true;
-
-	// Copy the cached body setup (unless we are creating it)
-	if (CachedBodySetup == nullptr)
-	{
-		CacheBodySetupHelper();
-	}
-	NewBS->CopyBodyPropertiesFrom(CachedBodySetup);
-
-	return NewBS;
-}
-
-UBodySetup* UMRMeshComponent::GetBodySetup()
-{
-	if (CachedBodySetup == nullptr)
-	{
-		CacheBodySetupHelper();
-	}
-	return CachedBodySetup;
-}
-
 void UMRMeshComponent::SendBrickData_Internal(IMRMesh::FSendBrickDataArgs Args)
 {
 	check(IsInGameThread());
@@ -763,104 +728,78 @@ void UMRMeshComponent::SendBrickData_Internal(IMRMesh::FSendBrickDataArgs Args)
 	
 	OnBrickDataUpdatedDelegate.Broadcast(this, Args);
 
-#if SUPPORTS_PHYSICS_COOKING
 	UE_LOG(LogMrMesh, Log, TEXT("SendBrickData_Internal() processing brick %llu with %i triangles"), Args.BrickId, Args.Indices.Num() / 3);
 	
+	// Collision generation.
 	if (IsValid(this) && !bNeverCreateCollisionMesh)
 	{
-		SCOPE_CYCLE_COUNTER(STAT_UpdateCollision);
-		// Physics update
-		UWorld* MyWorld = GetWorld();
-		if ( MyWorld && MyWorld->GetPhysicsScene() )
-		{
-			int32 BodyIndex = BodyIds.Find(Args.BrickId);
+			SCOPE_CYCLE_COUNTER(STAT_UpdateCollision);
 
-			if (bHasBrickData)
+			UWorld* World = GetWorld();
+			if (World && World->GetPhysicsScene())
 			{
-				bPhysicsStateCreated = true;
+				int32 BodyIndex = BodyIds.Find(Args.BrickId);
 
-				if (BodyIndex == INDEX_NONE)
+				if (bHasBrickData)
 				{
-					BodyIds.Add(Args.BrickId);
-					BodySetups.Add(CreateBodySetupHelper());
-					BodyInstances.Add(new FBodyInstance());
-					BodyIndex = BodyIds.Num() - 1;
-				}
+					bPhysicsStateCreated = true;
+
+					if (BodyIndex == INDEX_NONE)
+					{
+						TObjectPtr<UMRMeshBodyHolder> NewBodyHolder = NewObject<UMRMeshBodyHolder>(this, NAME_None, (IsTemplate() ? RF_Public | RF_ArchetypeObject : RF_NoFlags));
+						NewBodyHolder->Initialize(Args.BrickId);
+						BodyHolders.Add(NewBodyHolder);
+						BodyIndex = BodyIds.Add(Args.BrickId);
+					}
+
+					UMRMeshBodyHolder& BodyHolder = *BodyHolders[BodyIndex].Get();
+					check(BodyHolder.BrickId == Args.BrickId);
+
+					// Interrupt any in progress cook for this brick.
+					BodyHolder.AbortCook();
+
+					BodyHolder.Update(Args);
+
+					const bool bUseAsyncCook = bUseAsyncCooking && World->IsGameWorld();
+					if (bUseAsyncCook)
+					{
+						BodyHolder.bCookInProgress = true;
+						BodyHolder.BodySetup->CreatePhysicsMeshesAsync(FOnAsyncPhysicsCookFinished::CreateUObject(&BodyHolder, &UMRMeshBodyHolder::FinishPhysicsAsyncCook, BodyHolder.BodySetup.Get()));
+					}
+					else
+					{
+						// Also we want cooked data for this
+						BodyHolder.BodySetup->bHasCookedCollisionData = true;
+						BodyHolder.BodySetup->InvalidatePhysicsData();
+						BodyHolder.BodySetup->CreatePhysicsMeshes();
+						RecreatePhysicsState();
 
-				UBodySetup* MyBS = BodySetups[BodyIndex];
-				MyBS->bHasCookedCollisionData = true;
-				MyBS->CollisionTraceFlag = CTF_UseComplexAsSimple;
-				MyBS->ClearPhysicsMeshes();
-				MyBS->InvalidatePhysicsData();
-
-#if PHYSICS_INTERFACE_PHYSX
-				FCookBodySetupInfo CookInfo;
-				// Disable mesh cleaning by passing in EPhysXMeshCookFlags::DeformableMesh
-				static const EPhysXMeshCookFlags CookFlags = EPhysXMeshCookFlags::FastCook | EPhysXMeshCookFlags::DeformableMesh;
-				MyBS->GetCookInfo(CookInfo, CookFlags);
-				CookInfo.bCookTriMesh = true;
-				CookInfo.TriMeshCookFlags = CookInfo.ConvexCookFlags = CookFlags;
-				CookInfo.TriangleMeshDesc.bFlipNormals = true;
-				CookInfo.TriangleMeshDesc.Vertices = Args.PositionData;
-				const int NumFaces = Args.Indices.Num() / 3;
-				CookInfo.TriangleMeshDesc.Indices.Reserve(Args.Indices.Num() / 3);
-				for (int i = 0; i < NumFaces; ++i)
+						SuggestNavMeshUpdate();
+					}
+				}
+				else // !bHasBrickData
 				{
-					CookInfo.TriangleMeshDesc.Indices.AddUninitialized(1);
-					CookInfo.TriangleMeshDesc.Indices[i].v0 = Args.Indices[3 * i + 0];
-					CookInfo.TriangleMeshDesc.Indices[i].v1 = Args.Indices[3 * i + 1];
-					CookInfo.TriangleMeshDesc.Indices[i].v2 = Args.Indices[3 * i + 2];
+					if (BodyIndex != INDEX_NONE)
+					{
+						RemoveBodyInstance(BodyIndex);
+
+						SuggestNavMeshUpdate();
+					}
+					else
+					{
+						// This brick already doesn't exist, so no work to be done.
+					}
 				}
 
-				FPhysXCookHelper CookHelper(GetPhysXCookingModule());
-				CookHelper.CookInfo = CookInfo;
-				CookHelper.CreatePhysicsMeshes_Concurrent();
-
-				MyBS->FinishCreatingPhysicsMeshes_PhysX(CookHelper.OutNonMirroredConvexMeshes, CookHelper.OutMirroredConvexMeshes, CookHelper.OutTriangleMeshes);
-#else
-				// Chaos code path, save the temp pointers
-				TempPosition = &Args.PositionData;
-				TempIndices = &Args.Indices;
-				// It would be nice if we can call UBodySetup::CreatePhysicsMeshes directly, but currently that code path requires WITH_EDITOR
-				static const FName PhysicsFormatName(FPlatformProperties::GetPhysicsFormat());
-				// Build the collision data and saved it in a bulk data
-				FChaosDerivedDataCooker Cooker(MyBS, PhysicsFormatName);
-				TArray<uint8> OutData;
-				Cooker.Build(OutData);
-				FByteBulkData BulkData;
-				BulkData.Lock(LOCK_READ_WRITE);
-				FMemory::Memcpy(BulkData.Realloc(OutData.Num()), OutData.GetData(), OutData.Num());
-				BulkData.Unlock();
-				// Apply the collision data
-				MyBS->ProcessFormatData_Chaos(&BulkData);
-				// Clear the temp pointers since we don't own the data!
-				TempPosition = nullptr;
-				TempIndices = nullptr;
-#endif
-				FBodyInstance* MyBI = BodyInstances[BodyIndex];
-				MyBI->TermBody();
-				MyBI->InitBody(MyBS, GetComponentTransform(), this, MyWorld->GetPhysicsScene());
-				MyBI->CopyRuntimeBodyInstancePropertiesFrom(&BodyInstance);
 			}
 			else
 			{
-				if (BodyIndex != INDEX_NONE)
-				{
-					RemoveBodyInstance(BodyIndex);
-				}
-				else
-				{
-					// This brick already doesn't exist, so no work to be done.
-				}
+				// In A WORLD without physics you cannot COOK collision!
+				UE_LOG(LogMrMesh, Warning, TEXT("UMRMeshComponent::SendBrickData_Internal tried to cook collision without a world or in a world without physics!  No collision cooked."));
 			}
-		}
-		if (bHasBrickData && bUpdateNavMeshOnMeshUpdate && bHasCustomNavigableGeometry)
-		{
-			UpdateNavigationData();
-		}
 	}
-#endif // SUPPORTS_PHYSICS_COOKING
 
+	// Rendering generation
 	if (bCreateMeshProxySections)
 	{
 		if (SceneProxy != nullptr && GIsThreadedRendering)
@@ -888,10 +827,9 @@ void UMRMeshComponent::SendBrickData_Internal(IMRMesh::FSendBrickDataArgs Args)
 
 void UMRMeshComponent::RemoveBodyInstance(int32 BodyIndex)
 {
-	BodyInstances[BodyIndex]->TermBody();
-	delete BodyInstances[BodyIndex];
-	BodyInstances.RemoveAtSwap(BodyIndex);
-	BodySetups.RemoveAtSwap(BodyIndex);
+	auto& BH = BodyHolders[BodyIndex];
+	BH->Cleanup();
+	BodyHolders.RemoveAtSwap(BodyIndex);
 	BodyIds.RemoveAtSwap(BodyIndex);
 }
 
@@ -899,10 +837,10 @@ void UMRMeshComponent::OnUpdateTransform(EUpdateTransformFlags UpdateTransformFl
 {
 	Super::OnUpdateTransform(UpdateTransformFlags, Teleport);
 
-	for (FBodyInstance* BI : BodyInstances)
+	for (auto& BH : BodyHolders)
 	{
-		BI->SetBodyTransform(GetComponentTransform(), Teleport);
-		BI->UpdateBodyScale(GetComponentTransform().GetScale3D());
+		BH->BodyInstance.SetBodyTransform(GetComponentTransform(), Teleport);
+		BH->BodyInstance.UpdateBodyScale(GetComponentTransform().GetScale3D());
 	}
 }
 
@@ -983,19 +921,39 @@ bool UMRMeshComponent::DoCustomNavigableGeometryExport(FNavigableGeometryExport&
 {
 	check(bHasCustomNavigableGeometry);
 	
-	for (UBodySetup* BodySetup : BodySetups)
+	for (auto& BH : BodyHolders)
 	{
-		check(BodySetup);
-		GeomExport.ExportRigidBodySetup(*BodySetup, GetComponentTransform());
+		check(BH->BodySetup);
+		GeomExport.ExportRigidBodySetup(*BH->BodySetup, GetComponentTransform());
 	}
 
 	return false;
 }
 
+void UMRMeshComponent::SuggestNavMeshUpdate()
+{
+	bNavMeshUpdateSuggested = true;
+
+	if (bUpdateNavMeshOnMeshUpdate && bHasCustomNavigableGeometry)
+	{
+		bNavMeshUpdateSuggested = false;
+		UpdateNavigationData();
+	}
+}
+
+void UMRMeshComponent::RequestNavMeshUpdate()
+{
+	if (bNavMeshUpdateSuggested == true)
+	{
+		ForceNavMeshUpdate();
+	}
+}
+
 void UMRMeshComponent::ForceNavMeshUpdate()
 {
 	if (bHasCustomNavigableGeometry)
 	{
+		bNavMeshUpdateSuggested = false;
 		UpdateNavigationData();
 	}
 	else
@@ -1141,37 +1099,109 @@ UMaterialInterface* UMRMeshComponent::GetMaterialToUse() const
 	}
 }
 
-bool UMRMeshComponent::GetPhysicsTriMeshData(struct FTriMeshCollisionData* CollisionData, bool InUseAllTriData)
+void UMRMeshBodyHolder::Initialize(IMRMesh::FBrickId InBrickId)
+{
+	check(BodySetup == nullptr);
+
+	BrickId = InBrickId;
+
+	BodySetup = NewObject<UBodySetup>(this, NAME_None, (IsTemplate() ? RF_Public | RF_ArchetypeObject : RF_NoFlags));
+	BodySetup->BodySetupGuid = FGuid::NewGuid();
+	BodySetup->bGenerateMirroredCollision = false;
+	BodySetup->bHasCookedCollisionData = true;
+	//BodySetup->bDoubleSidedGeometry = true;
+	BodySetup->CollisionTraceFlag = CTF_UseComplexAsSimple;
+
+	bCookInProgress = false;
+}
+
+void UMRMeshBodyHolder::Update(const IMRMesh::FSendBrickDataArgs& Args)
+{
+	check(BodySetup != nullptr);
+	check(BrickId == Args.BrickId);
+	BrickDataReceipt = Args.BrickDataReceipt; // This would release hold on any previous receipt.
+	PositionData = &Args.PositionData;
+	Indices = &Args.Indices;
+	Bounds = Args.Bounds;
+}
+
+void UMRMeshBodyHolder::AbortCook()
+{
+	if (bCookInProgress)
+	{
+		BodySetup->AbortPhysicsMeshAsyncCreation();
+		bCookInProgress = false;
+
+		ReleaseArgData();
+	}
+}
+
+void UMRMeshBodyHolder::ReleaseArgData()
+{
+	// Null all the raw pointers and release our hold on the receipt.
+	PositionData = nullptr;
+	Indices = nullptr;
+	BrickDataReceipt.Reset();
+
+	// Leave the POD members alone.
+}
+
+void UMRMeshBodyHolder::Cleanup()
+{
+	AbortCook();
+	ReleaseArgData();
+
+	BodyInstance.TermBody();
+}
+
+bool UMRMeshBodyHolder::GetPhysicsTriMeshData(struct FTriMeshCollisionData* CollisionData, bool InUseAllTriData)
 {
-	if (TempPosition && TempIndices)
+	if (BrickDataReceipt.IsValid() && Indices->Num() > 0)
 	{
 		// Copy the vertices
-		CollisionData->Vertices = *TempPosition;
-		
+		CollisionData->Vertices.AddUninitialized(PositionData->Num());
+		for (auto Position = 0; Position < PositionData->Num(); ++Position)
+		{
+			CollisionData->Vertices[Position] = (*PositionData)[Position];
+		}
+
 		// Copy the indices
-		const auto& Indices = *TempIndices;
-		CollisionData->Indices.Reset(Indices.Num() / 3);
-		for (auto Index = 0; Index < Indices.Num(); Index += 3)
+		const auto& IndicesR = *Indices;
+		CollisionData->Indices.Reset(IndicesR.Num() / 3);
+		for (auto Index = 0; Index < IndicesR.Num(); Index += 3)
 		{
 			FTriIndices Face;
-			Face.v0 = Indices[Index];
-			Face.v1 = Indices[Index + 1];
-			Face.v2 = Indices[Index + 2];
+			Face.v0 = IndicesR[Index];
+			Face.v1 = IndicesR[Index + 1];
+			Face.v2 = IndicesR[Index + 2];
 			CollisionData->Indices.Add(Face);
 		}
-		
+
+		CollisionData->bFlipNormals = true;
+		CollisionData->bDeformableMesh = true;
+		CollisionData->bFastCook = true;
+
 		return true;
 	}
-	
+
 	return false;
 }
+bool UMRMeshBodyHolder::ContainsPhysicsTriMeshData(bool InUseAllTriData) const
+{
+	return (BrickDataReceipt.IsValid() && Indices->Num() > 0);
+}
 
-bool UMRMeshComponent::ContainsPhysicsTriMeshData(bool InUseAllTriData) const
+/** Once async physics cook is done, create needed state */
+void UMRMeshBodyHolder::FinishPhysicsAsyncCook(bool bSuccess, UBodySetup* FinishedBodySetup)
 {
-	if (TempPosition && TempIndices)
+	UMRMeshComponent* MRMeshComponent = Cast<UMRMeshComponent>(GetOuter());
+	if (MRMeshComponent)
 	{
-		return true;
+		FTransform BodyTransform = MRMeshComponent->GetComponentTransform();
+		BodyInstance.TermBody();
+		BodyInstance.InitBody(BodySetup, BodyTransform, MRMeshComponent, GetWorld()->GetPhysicsScene());
+		BodyInstance.CopyRuntimeBodyInstancePropertiesFrom(MRMeshComponent->GetBodyInstance());
+
+		MRMeshComponent->SuggestNavMeshUpdate();
 	}
-	
-	return false;
-}
+}
\ No newline at end of file
diff --git a/Engine/Source/Runtime/MRMesh/Private/MockDataMeshTrackerComponent.cpp b/Engine/Source/Runtime/MRMesh/Private/MockDataMeshTrackerComponent.cpp
--- a/Engine/Source/Runtime/MRMesh/Private/MockDataMeshTrackerComponent.cpp
+++ b/Engine/Source/Runtime/MRMesh/Private/MockDataMeshTrackerComponent.cpp
@@ -425,9 +425,11 @@ void UMockDataMeshTrackerComponent::RemoveBlock(int32 BlockIndex)
 void UMockDataMeshTrackerComponent::UpdateBlock(int32 BlockIndex)
 {
 	// Create a brick index for any new mesh block
+	bool bBlockIsNew = false;
 	if (!Impl->MeshBrickCache.Contains(BlockIndex))
 	{
 		Impl->MeshBrickCache.Add(BlockIndex, Impl->MeshBrickIndex++);
+		bBlockIsNew = true;
 	}
 
 	const FMockDataMeshTrackerImpl::FRawMockMeshData& RawMeshData = Impl->RawMockMeshData[BlockIndex];
@@ -438,14 +440,15 @@ void UMockDataMeshTrackerComponent::UpdateBlock(int32 BlockIndex)
 	CurrentMeshDataCache->BrickId = BrickId;
 
 	// Pull vertices
+	const FVector UpdateShiftVector = bBlockIsNew ? FVector::ZeroVector : FVector(0.0f, 0.0f, 15.0f);  // when we update we will shift the verts up a little, so we can see that something happened.
 	const FVector VertexOffset = FVector::ZeroVector; // If this was the inverse of the worldToTrackingTransform position the mesh would be located in tracker space, but we don't want a dependency on HMD here.  
 	const int32 VertexCount = RawMeshData.Vertices.Num();
 	CurrentMeshDataCache->OffsetVertices.Reserve(VertexCount);
 	CurrentMeshDataCache->WorldVertices.Reserve(VertexCount);
 	for (int32 v = 0; v < VertexCount; ++ v)
 	{
-			CurrentMeshDataCache->OffsetVertices.Add(RawMeshData.Vertices[v] - VertexOffset);
-			CurrentMeshDataCache->WorldVertices.Add(RawMeshData.Vertices[v]);
+			CurrentMeshDataCache->OffsetVertices.Add(RawMeshData.Vertices[v] - VertexOffset + UpdateShiftVector);
+			CurrentMeshDataCache->WorldVertices.Add(RawMeshData.Vertices[v] + UpdateShiftVector);
 	}
 
 	// Pull indices
diff --git a/Engine/Source/Runtime/MRMesh/Public/MRMeshComponent.h b/Engine/Source/Runtime/MRMesh/Public/MRMeshComponent.h
--- a/Engine/Source/Runtime/MRMesh/Public/MRMeshComponent.h
+++ b/Engine/Source/Runtime/MRMesh/Public/MRMeshComponent.h
@@ -1,5 +1,12 @@
 // Copyright Epic Games, Inc. All Rights Reserved.
 
+// MR stands for Mixed Reality, but this component could possibly have more general use.
+// It was specifically designed to support the case for Augmented Reality world mesh detection.  Where a device scans the real world and provides a mesh representing the geometry
+// of the real world environment.  This mesh is expected to arrive in chunks called 'bricks' over time and for bricks to be added, removed, or updated.
+// We provide the ability to render the mesh, to generate collision from it, and to generate ai navigation data on that collision.
+// We do not intend to support using the MRMesh as a physics body that could be moved.
+// It is somewhat similar to ProceduralMeshComponent.
+
 #pragma once
 
 #include "CoreMinimal.h"
@@ -36,8 +43,9 @@ public:
 	typedef uint64 FBrickId;
 	struct FSendBrickDataArgs
 	{
+		// Note: many members are references!  Be aware of the lifetime of the underlying data.  FBrickDataReceipt is intended to help manage that.
 		TSharedPtr<FBrickDataReceipt, ESPMode::ThreadSafe> BrickDataReceipt;
-		const FBrickId BrickId;
+		const FBrickId BrickId = 0;
 		const TArray<FVector3f>& PositionData;
 		const TArray<FVector2D>& UVData;
 		const TArray<FPackedNormal>& TangentXZData;
@@ -55,21 +63,56 @@ public:
 	virtual void ClearAllBrickData() = 0;
 };
 
+// Because physics cooking uses GetOuter() to get the IInterface_CollisionDataProvider and provides no way to determine which physics body it
+// is currently working on we are wrapping each body in this Holder so that it can be the Outer and provide the correct data.
+UCLASS(transient)
+class UMRMeshBodyHolder : public UObject, public IInterface_CollisionDataProvider
+{
+public:
+	GENERATED_BODY()
+
+	void Initialize(IMRMesh::FBrickId InBrickId);
+	void Update(const IMRMesh::FSendBrickDataArgs& Args);
+	void AbortCook();
+	void ReleaseArgData();
+	void Cleanup();
+
+	//~ Begin Interface_CollisionDataProvider Interface
+	virtual bool GetPhysicsTriMeshData(struct FTriMeshCollisionData* CollisionData, bool InUseAllTriData) override;
+	virtual bool ContainsPhysicsTriMeshData(bool InUseAllTriData) const override;
+	virtual bool WantsNegXTriMesh() override { return false; }
+	//~ End Interface_CollisionDataProvider Interface
+
+	/** Once async physics cook is done, create needed state */
+	void FinishPhysicsAsyncCook(bool bSuccess, UBodySetup* FinishedBodySetup);
+
+	UPROPERTY(transient)
+	TObjectPtr<class UBodySetup> BodySetup;
+
+	UPROPERTY(transient)
+	FBodyInstance BodyInstance;
+
+	// Hold references to the bulk data which is owned externally, like FSendBrickDataArgs.  The receipt keeps the references valid.
+	TSharedPtr<IMRMesh::FBrickDataReceipt, ESPMode::ThreadSafe> BrickDataReceipt;
+	IMRMesh::FBrickId BrickId = 0;
+	const TArray<FVector3f>* PositionData = nullptr;
+	const TArray<MRMESH_INDEX_TYPE>* Indices = nullptr;
+	FBox Bounds = FBox(EForceInit::ForceInit);
+
+	bool bCookInProgress = false;
+};
+
+
 
 DECLARE_MULTICAST_DELEGATE_TwoParams(FOnMRMeshBrickDataUpdatedDelegate, const UMRMeshComponent*, const IMRMesh::FSendBrickDataArgs&);
 
-UCLASS(hideCategories=(Physics), meta = (BlueprintSpawnableComponent, Experimental), ClassGroup = Rendering)
-class MRMESH_API UMRMeshComponent : public UPrimitiveComponent, public IMRMesh, public IInterface_CollisionDataProvider
+UCLASS(hideCategories=(Physics), meta = (BlueprintSpawnableComponent), ClassGroup = Rendering)
+class MRMESH_API UMRMeshComponent : public UPrimitiveComponent, public IMRMesh
 {
 public:
 	friend class FMRMeshProxy;
 
 	GENERATED_UCLASS_BODY()
-	
-	//~ IInterface_CollisionDataProvider
-	virtual bool GetPhysicsTriMeshData(struct FTriMeshCollisionData* CollisionData, bool InUseAllTriData) override;
-	virtual bool ContainsPhysicsTriMeshData(bool InUseAllTriData) const override;
-	//~ IInterface_CollisionDataProvider
 
 	virtual void BeginPlay() override;
 	void EndPlay(const EEndPlayReason::Type EndPlayReason) override;
@@ -80,9 +123,18 @@ public:
 	void SetConnected(bool value) override { bConnected = value; }
 	virtual void SendRelativeTransform(const FTransform& Transform) override { SetRelativeTransform(Transform); }
 
+	/**
+	* Force navmesh generation to run using the current collision data.  This will run even if the collision data has not been udpated! Unless you are changing navmesh settings or similar RequestNavMeshUpdate is reccomended.
+	*/	
 	UFUNCTION(BlueprintCallable, Category = "Mesh Reconstruction")
 	void ForceNavMeshUpdate();
 
+	/**
+	* Generate nav mesh if collision data has changed since the last nav mesh generation.  
+	*/
+	UFUNCTION(BlueprintCallable, Category = "Mesh Reconstruction")
+	void RequestNavMeshUpdate();
+
 	UFUNCTION(BlueprintCallable, Category = "Mesh Reconstruction")
 	void Clear() override;
 
@@ -134,6 +186,7 @@ public:
 
 	void SetNeverCreateCollisionMesh(bool bNeverCreate) { bNeverCreateCollisionMesh = bNeverCreate; }
 	void SetEnableNavMesh(bool bEnable) { bUpdateNavMeshOnMeshUpdate = bEnable;  }
+	void SuggestNavMeshUpdate();
 
 	/** Trackers feeding mesh data to this component may want to know when we clear our mesh data */
 	DECLARE_EVENT(UMRMeshComponent, FOnClear);
@@ -145,7 +198,7 @@ private:
 	virtual void GetUsedMaterials(TArray<UMaterialInterface*>& OutMaterials, bool bGetDebugMaterials = false) const override;
 	virtual FBoxSphereBounds CalcBounds(const FTransform& LocalToWorld) const override;
 	virtual bool DoCustomNavigableGeometryExport(FNavigableGeometryExport& GeomExport) const override;
-	virtual class UBodySetup* GetBodySetup();
+
 	//~ UPrimitiveComponent
 	//~ UActorComponent
 	virtual bool ShouldCreatePhysicsState() const override;
@@ -158,8 +211,6 @@ private:
 	//~ IMRMesh
 
 private:
-	void CacheBodySetupHelper();
-	class UBodySetup* CreateBodySetupHelper();
 	void SendBrickData_Internal(IMRMesh::FSendBrickDataArgs Args);
 
 	void RemoveBodyInstance(int32 BodyIndex);
@@ -181,34 +232,36 @@ private:
 	UPROPERTY(EditAnywhere, Category = MRMesh)
 	bool bUpdateNavMeshOnMeshUpdate = true;
 
+	bool bNavMeshUpdateSuggested = false;
+
 	/** If true, MRMesh will not create a collidable ridgid body for each mesh section and can therefore never have collision.  Avoids the cost of generating collision.*/
 	UPROPERTY(EditAnywhere, Category = MRMesh)
 	bool bNeverCreateCollisionMesh = false;
 
-	bool bConnected = false;
-
-	UPROPERTY(Transient)
-	TObjectPtr<class UBodySetup> CachedBodySetup = nullptr;
-
-	UPROPERTY(Transient)
-	TArray<TObjectPtr<UBodySetup>> BodySetups;
+	/**
+	*	Controls whether the physics cooking should be done off the game thread. This should be used when collision geometry doesn't have to be immediately up to date (For example streaming in far away objects).  
+	*   Fixing this to true for now.
+	*/
+	bool bUseAsyncCooking = true;
 
+	bool bConnected = false;
 
 	/** Whether this mesh should write z-depth to occlude meshes or not */
 	bool bEnableOcclusion;
 	/** Whether this mesh should draw using the wireframe material when no material is set or not */
 	bool bUseWireframe;
 
-	TArray<FBodyInstance*> BodyInstances;
-	TArray<IMRMesh::FBrickId> BodyIds;
-
 	FOnClear OnClearEvent;
 	
 	FLinearColor WireframeColor = FLinearColor::White;
-	
-	/** Temporarily saved data pointers used for constructing the collision mesh */
-	const TArray<FVector3f>* TempPosition = nullptr;
-	const TArray<MRMESH_INDEX_TYPE>* TempIndices = nullptr;
-	
+
 	FOnMRMeshBrickDataUpdatedDelegate OnBrickDataUpdatedDelegate;
+
+
+	// Collision/Physics data
+	UPROPERTY(Transient)
+	TArray<TObjectPtr<UMRMeshBodyHolder>> BodyHolders;
+
+	// This array is parallel to BodyHolders
+	TArray<IMRMesh::FBrickId> BodyIds;
 };
