UE-91500 Hololens generated mesh work once PIE with TP_HolloensBP Template
-Fixed world mesh with hololens on WMR plugin.  WMR interop writes floats UE side was assuming doubles.  Now we assume floats in the affected buffers.

diff --git a/Engine/Plugins/Runtime/WindowsMixedReality/Source/HoloLensAR/Private/HoloLensARSystem.cpp b/Engine/Plugins/Runtime/WindowsMixedReality/Source/HoloLensAR/Private/HoloLensARSystem.cpp
--- a/Engine/Plugins/Runtime/WindowsMixedReality/Source/HoloLensAR/Private/HoloLensARSystem.cpp
+++ b/Engine/Plugins/Runtime/WindowsMixedReality/Source/HoloLensAR/Private/HoloLensARSystem.cpp
@@ -53,9 +53,9 @@ struct FMeshUpdate
 
 	// These will use MoveTemp to avoid another copy
 	// The interop fills these in directly
-	TArray<FVector> Vertices;
+	TArray<FVector3f> Vertices;
 	TArray<MRMESH_INDEX_TYPE> Indices;
-	TArray<FVector> Normals;
+	TArray<FVector3f> Normals;
 
 	bool bIsRightHandMesh;
 };
@@ -1487,8 +1487,16 @@ bool FHoloLensARSystem::GetHandMeshData(EControllerHand Hand, TArray<FVector>& O
 	OutVertices.AddUninitialized(HandState.Vertices.Num());
 	OutNormals.AddUninitialized(HandState.Normals.Num());
 
-	OutVertices = CopyTemp(HandState.Vertices);
-	OutNormals = CopyTemp(HandState.Normals);
+	const int Num = HandState.Vertices.Num();
+	for (int i = 0; i < Num; i++)
+	{
+		OutVertices[i] = HandState.Vertices[i];
+	}
+	const int Num2 = HandState.Normals.Num();
+	for (int i = 0; i < Num2; i++)
+	{
+		OutNormals[i] = HandState.Normals[i];
+	}	
 
 	auto DestIndices = OutIndices.GetData();
 	for (size_t i = 0; i < HandState.Indices.Num(); i++)
diff --git a/Engine/Source/Runtime/MRMesh/Private/MRMeshComponent.cpp b/Engine/Source/Runtime/MRMesh/Private/MRMeshComponent.cpp
--- a/Engine/Source/Runtime/MRMesh/Private/MRMeshComponent.cpp
+++ b/Engine/Source/Runtime/MRMesh/Private/MRMeshComponent.cpp
@@ -979,11 +979,51 @@ struct FMeshArrayHolder :
 	TArray<FPackedNormal> BogusTangents;
 	TArray<FColor> BogusColors;
 
+	FMeshArrayHolder(TArray<FVector3f>& InVertices, TArray<MRMESH_INDEX_TYPE>& InIndices, TArray<FVector2D>& UVData, TArray<FPackedNormal>& TangentXZData, TArray<FColor>& ColorData)
+		: Indices(MoveTemp(InIndices))
+	{
+		const int32 CurrentNumVertices = InVertices.Num();
+
+		// This constructor is a bit slower because it has to copy the vertices one by one to convert them from float to double.
+		Vertices.AddUninitialized(CurrentNumVertices);
+		for (int i = 0; i < CurrentNumVertices; i++)
+		{
+			Vertices[i] = InVertices[i];
+		}
+
+		if (UVData.Num() == CurrentNumVertices)
+		{
+			BogusUVs = MoveTemp(UVData);
+		}
+		else
+		{
+			BogusUVs.AddZeroed(CurrentNumVertices);
+		}
+
+		if (ColorData.Num() == CurrentNumVertices)
+		{
+			BogusColors = MoveTemp(ColorData);
+		}
+		else
+		{
+			BogusColors.AddZeroed(CurrentNumVertices);
+		}
+
+		if (TangentXZData.Num() == CurrentNumVertices * 2)
+		{
+			BogusTangents = MoveTemp(TangentXZData);
+		}
+		else
+		{
+			BogusTangents.AddZeroed(CurrentNumVertices * 2);
+		}
+	}
+
 	FMeshArrayHolder(TArray<FVector>& InVertices, TArray<MRMESH_INDEX_TYPE>& InIndices, TArray<FVector2D>& UVData, TArray<FPackedNormal>& TangentXZData, TArray<FColor>& ColorData)
 		: Vertices(MoveTemp(InVertices))
 		, Indices(MoveTemp(InIndices))
 	{
-		int32 CurrentNumVertices = Vertices.Num();
+		const int32 CurrentNumVertices = Vertices.Num();
 		
 		if (UVData.Num() == CurrentNumVertices)
 		{
@@ -1014,6 +1054,38 @@ struct FMeshArrayHolder :
 	}
 };
 
+void UMRMeshComponent::UpdateMesh(const FVector& InLocation, const FQuat& InRotation, const FVector& Scale, TArray<FVector3f>& Vertices, TArray<MRMESH_INDEX_TYPE>& Indices, TArray<FVector2D> UVData, TArray<FPackedNormal> TangentXZData, TArray<FColor> ColorData)
+{
+	SetRelativeLocationAndRotation(InLocation, InRotation);
+	SetRelativeScale3D(Scale);
+
+	// Create our struct that will hold the data until the render thread is done with it
+	TSharedPtr<FMeshArrayHolder, ESPMode::ThreadSafe> MeshHolder = MakeShared<FMeshArrayHolder>(Vertices, Indices, UVData, TangentXZData, ColorData);
+	// NOTE: Indices are empty due to MoveTemp()!!!
+	 
+	// Set a valid bounding box so meshes get correctly culled.
+	FBox bounds = FBox(ForceInit);
+	const int Num = Vertices.Num();
+	for (int i = 0; i < Num; i++)
+	{
+		bounds += Vertices[i];
+	}
+	bounds = bounds.TransformBy(FTransform(InRotation.Rotator(), InLocation, Scale));
+
+	SendBrickData_Internal(IMRMesh::FSendBrickDataArgs
+		{
+			MeshHolder,
+			0,
+			MeshHolder->Vertices,
+			MeshHolder->BogusUVs,
+			MeshHolder->BogusTangents,
+			MeshHolder->BogusColors,
+			MeshHolder->Indices,
+			bounds
+		}
+	);
+}
+
 void UMRMeshComponent::UpdateMesh(const FVector& InLocation, const FQuat& InRotation, const FVector& Scale, TArray<FVector>& Vertices, TArray<MRMESH_INDEX_TYPE>& Indices, TArray<FVector2D> UVData, TArray<FPackedNormal> TangentXZData, TArray<FColor> ColorData)
 {
 	SetRelativeLocationAndRotation(InLocation, InRotation);
diff --git a/Engine/Source/Runtime/MRMesh/Public/MRMeshComponent.h b/Engine/Source/Runtime/MRMesh/Public/MRMeshComponent.h
--- a/Engine/Source/Runtime/MRMesh/Public/MRMeshComponent.h
+++ b/Engine/Source/Runtime/MRMesh/Public/MRMeshComponent.h
@@ -148,6 +148,7 @@ public:
 
 	/** Updates from HoloLens or iOS */
 	void UpdateMesh(const FVector& InLocation, const FQuat& InRotation, const FVector& Scale, TArray<FVector>& Vertices, TArray<MRMESH_INDEX_TYPE>& Indices, TArray<FVector2D> UVData = {}, TArray<FPackedNormal> TangentXZData = {}, TArray<FColor> ColorData = {});
+	void UpdateMesh(const FVector& InLocation, const FQuat& InRotation, const FVector& Scale, TArray<FVector3f>& Vertices, TArray<MRMESH_INDEX_TYPE>& Indices, TArray<FVector2D> UVData = {}, TArray<FPackedNormal> TangentXZData = {}, TArray<FColor> ColorData = {});
 
 	UFUNCTION(BlueprintCallable, Category = "Mesh Reconstruction")
 	void SetEnableMeshOcclusion(bool bEnable);
