UE-168661 //UE5/Release-5.1 - UE.EditorAutomation(RunTest=Project.Maps.AllInPIE) Win64 Development EditorGame - QAEntry test fails due to Unhandled Exception in MockDataMeshTrackerComponent.cpp
-Fixed a bug where we would hold a pointer to the component beyond its lifetime and then try to use it.  Now we clear that pointer in the component's BeginDestroy.

robomerge[FN_EngineMerge] -UE5-RES

diff --git a/Engine/Source/Runtime/MRMesh/Private/MockDataMeshTrackerComponent.cpp b/Engine/Source/Runtime/MRMesh/Private/MockDataMeshTrackerComponent.cpp
--- a/Engine/Source/Runtime/MRMesh/Private/MockDataMeshTrackerComponent.cpp
+++ b/Engine/Source/Runtime/MRMesh/Private/MockDataMeshTrackerComponent.cpp
@@ -137,9 +137,15 @@ public:
 		TArray<MRMESH_INDEX_TYPE> Triangles;
 		TArray<float> Confidence;
 
+		// Disconnect this FCachedMeshData so if it is deleted after FMockDataMeshTrackerImpl is deleted it does not try to call back in to add itself to the free list.
+		void Disconnect()
+		{
+			Owner = nullptr;
+		}
+
+		// Reset this FCachedMeshData so that it can be reused later.
 		void Recycle(SharedPtr& MeshData)
 		{
-			check(Owner);
 			FMockDataMeshTrackerImpl* TempOwner = Owner;
 			Owner = nullptr;
 	
@@ -153,7 +159,10 @@ public:
 			Tangents.Reset();
 			Confidence.Reset();
 	
-			TempOwner->FreeMeshDataCache(MeshData);
+			if (TempOwner)
+			{
+				TempOwner->FreeMeshDataCache(MeshData);
+			}
 		}
 
 		void Init(FMockDataMeshTrackerImpl* InOwner)
@@ -212,6 +221,18 @@ public:
 		return true;
 	}
 
+	void BeginDestroy()
+	{
+		// Drop all the CachedMeshDatas references so that they will be destroyed unless a receipt is still hanging onto them.
+		// Disconnect them so that they do not try to call back into this to be put in the FreeCachedMeshDatas array.
+		for (FCachedMeshData::SharedPtr& Data : CachedMeshDatas)
+		{
+			Data->Disconnect();
+		}
+		CachedMeshDatas.Empty();
+		FreeCachedMeshDatas.Empty();
+	}
+
 	void Destroy()
 	{
 	}
@@ -605,6 +626,8 @@ void UMockDataMeshTrackerComponent::UpdateBlock(int32 BlockIndex)
 
 void UMockDataMeshTrackerComponent::BeginDestroy()
 {
+	Impl->BeginDestroy();
+
 	if (MRMesh != nullptr)
 	{
 		DisconnectMRMesh(MRMesh);
