[Engine]

diff --git a/Engine/Source/Runtime/LevelSequence/Public/LevelSequenceDirector.h b/Engine/Source/Runtime/LevelSequence/Public/LevelSequenceDirector.h
--- a/Engine/Source/Runtime/LevelSequence/Public/LevelSequenceDirector.h
+++ b/Engine/Source/Runtime/LevelSequence/Public/LevelSequenceDirector.h
@@ -13,22 +13,22 @@
 class IMovieScenePlayer;
 class ULevelSequencePlayer;
 
-UCLASS(Blueprintable)
-class LEVELSEQUENCE_API ULevelSequenceDirector : public UObject
+UCLASS(Blueprintable, MinimalAPI)
+class ULevelSequenceDirector : public UObject
 {
 public:
 	GENERATED_BODY()
 
 	/** Called when this director is created */
 	UFUNCTION(BlueprintImplementableEvent, Category="Sequencer")
-	void OnCreated();
+	LEVELSEQUENCE_API void OnCreated();
 	
 	/**
 	 * Get the current time for the outermost (root) sequence
 	 * @return The current playback position of the outermost (root) sequence
 	 */
 	UFUNCTION(BlueprintCallable, Category = "Sequencer|Director")
-	FQualifiedFrameTime GetRootSequenceTime() const;
+	LEVELSEQUENCE_API FQualifiedFrameTime GetRootSequenceTime() const;
 
 	UE_DEPRECATED(5.2, "GetMasterSequenceTime is deprecated. Please use GetRootSequenceTime instead")
 	UFUNCTION(BlueprintCallable, Category = "Sequencer|Director", meta = (DeprecatedFunction, DeprecationMessage = "Use GetRootSequenceTime"))
@@ -39,7 +39,7 @@ public:
 	 * @return The current playback position of this director's sequence
 	 */
 	UFUNCTION(BlueprintCallable, Category = "Sequencer|Director")
-	FQualifiedFrameTime GetCurrentTime() const;
+	LEVELSEQUENCE_API FQualifiedFrameTime GetCurrentTime() const;
 
 	/**
 	 * Resolve the bindings inside this sub-sequence that relate to the specified ID
@@ -48,7 +48,7 @@ public:
 	 * @param ObjectBinding The ID for the object binding inside this sub-sequence or one of its children to resolve
 	 */
 	UFUNCTION(BlueprintCallable, Category="Sequencer|Director")
-	TArray<UObject*> GetBoundObjects(FMovieSceneObjectBindingID ObjectBinding);
+	LEVELSEQUENCE_API TArray<UObject*> GetBoundObjects(FMovieSceneObjectBindingID ObjectBinding);
 
 
 	/**
@@ -58,7 +58,7 @@ public:
 	 * @param ObjectBinding The ID for the object binding inside this sub-sequence or one of its children to resolve
 	 */
 	UFUNCTION(BlueprintCallable, Category="Sequencer|Director")
-	UObject* GetBoundObject(FMovieSceneObjectBindingID ObjectBinding);
+	LEVELSEQUENCE_API UObject* GetBoundObject(FMovieSceneObjectBindingID ObjectBinding);
 
 
 	/**
@@ -68,7 +68,7 @@ public:
 	 * @param ObjectBinding The ID for the object binding inside this sub-sequence or one of its children to resolve
 	 */
 	UFUNCTION(BlueprintCallable, Category="Sequencer|Director")
-	TArray<AActor*> GetBoundActors(FMovieSceneObjectBindingID ObjectBinding);
+	LEVELSEQUENCE_API TArray<AActor*> GetBoundActors(FMovieSceneObjectBindingID ObjectBinding);
 
 
 	/**
@@ -78,17 +78,17 @@ public:
 	 * @param ObjectBinding The ID for the object binding inside this sub-sequence or one of its children to resolve
 	 */
 	UFUNCTION(BlueprintCallable, Category="Sequencer|Director")
-	AActor* GetBoundActor(FMovieSceneObjectBindingID ObjectBinding);
+	LEVELSEQUENCE_API AActor* GetBoundActor(FMovieSceneObjectBindingID ObjectBinding);
 
 	/*
 	 * Get the current sequence that this director is playing back within 
 	 */
 	UFUNCTION(BlueprintCallable, Category="Sequencer|Director")
-	UMovieSceneSequence* GetSequence();
+	LEVELSEQUENCE_API UMovieSceneSequence* GetSequence();
 
 public:
 
-	virtual UWorld* GetWorld() const override;
+	LEVELSEQUENCE_API virtual UWorld* GetWorld() const override;
 
 	/** Pointer to the player that's playing back this director's sequence. Only valid in game or in PIE/Simulate. */
 	UPROPERTY(BlueprintReadOnly, Category="Cinematics")
@@ -114,4 +114,4 @@ class ULegacyLevelSequenceDirectorBlueprint : public UBlueprint
 	{
 		ParentClass = ULevelSequenceDirector::StaticClass();
 	}
-};
\ No newline at end of file
+};
diff --git a/Engine/Source/Runtime/LevelSequence/Public/LevelSequencePlayer.h b/Engine/Source/Runtime/LevelSequence/Public/LevelSequencePlayer.h
--- a/Engine/Source/Runtime/LevelSequence/Public/LevelSequencePlayer.h
+++ b/Engine/Source/Runtime/LevelSequence/Public/LevelSequencePlayer.h
@@ -80,12 +80,12 @@ private:
  * This class keeps track of playback state and provides functions for manipulating
  * an level sequence while its playing.
  */
-UCLASS(BlueprintType)
-class LEVELSEQUENCE_API ULevelSequencePlayer
+UCLASS(BlueprintType, MinimalAPI)
+class ULevelSequencePlayer
 	: public UMovieSceneSequencePlayer
 {
 public:
-	ULevelSequencePlayer(const FObjectInitializer&);
+	LEVELSEQUENCE_API ULevelSequencePlayer(const FObjectInitializer&);
 
 	GENERATED_BODY()
 
@@ -96,7 +96,7 @@ public:
 	 * @param InLevel The level that the animation is played in.
 	 * @param InCameraSettings The desired camera settings
 	 */
-	void Initialize(ULevelSequence* InLevelSequence, ULevel* InLevel, const FLevelSequenceCameraSettings& InCameraSettings);
+	LEVELSEQUENCE_API void Initialize(ULevelSequence* InLevelSequence, ULevel* InLevel, const FLevelSequenceCameraSettings& InCameraSettings);
 
 	UE_DEPRECATED(5.1, "Use SetPlaybackSettings(...) then Initialize(ULevelSequence*, ULevel*, const FLevelSequenceCameraSettings&)")
 	void Initialize(ULevelSequence* InLevelSequence, ULevel* InLevel, const FMovieSceneSequencePlaybackSettings& Settings, const FLevelSequenceCameraSettings& InCameraSettings)
@@ -116,7 +116,7 @@ public:
 	 * @param OutActor The level sequence actor created to play this sequence.
 	 */
 	UFUNCTION(BlueprintCallable, Category="Sequencer|Player", meta=(WorldContext="WorldContextObject", DynamicOutputParam="OutActor"))
-	static ULevelSequencePlayer* CreateLevelSequencePlayer(UObject* WorldContextObject, ULevelSequence* LevelSequence, FMovieSceneSequencePlaybackSettings Settings, ALevelSequenceActor*& OutActor);
+	static LEVELSEQUENCE_API ULevelSequencePlayer* CreateLevelSequencePlayer(UObject* WorldContextObject, ULevelSequence* LevelSequence, FMovieSceneSequencePlaybackSettings Settings, ALevelSequenceActor*& OutActor);
 
 	/** Event triggered when there is a camera cut */
 	UPROPERTY(BlueprintAssignable, Category="Sequencer|Player")
@@ -129,37 +129,37 @@ public:
 public:
 
 	// IMovieScenePlayer interface
-	virtual UObject* GetPlaybackContext() const override;
-	virtual TArray<UObject*> GetEventContexts() const override;
+	LEVELSEQUENCE_API virtual UObject* GetPlaybackContext() const override;
+	LEVELSEQUENCE_API virtual TArray<UObject*> GetEventContexts() const override;
 
-	void RewindForReplay();
+	LEVELSEQUENCE_API void RewindForReplay();
 
 protected:
 
 	// IMovieScenePlayer interface
-	virtual void UpdateCameraCut(UObject* CameraObject, const EMovieSceneCameraCutParams& CameraCutParams) override;
-	virtual void ResolveBoundObjects(const FGuid& InBindingId, FMovieSceneSequenceID SequenceID, UMovieSceneSequence& InSequence, UObject* ResolutionContext, TArray<UObject*, TInlineAllocator<1>>& OutObjects) const override;
+	LEVELSEQUENCE_API virtual void UpdateCameraCut(UObject* CameraObject, const EMovieSceneCameraCutParams& CameraCutParams) override;
+	LEVELSEQUENCE_API virtual void ResolveBoundObjects(const FGuid& InBindingId, FMovieSceneSequenceID SequenceID, UMovieSceneSequence& InSequence, UObject* ResolutionContext, TArray<UObject*, TInlineAllocator<1>>& OutObjects) const override;
 
 	//~ UMovieSceneSequencePlayer interface
-	virtual bool CanPlay() const override;
-	virtual void OnStartedPlaying() override;
-	virtual void OnStopped() override;
-	virtual void UpdateMovieSceneInstance(FMovieSceneEvaluationRange InRange, EMovieScenePlayerStatus::Type PlayerStatus, const FMovieSceneUpdateArgs& Args) override;
+	LEVELSEQUENCE_API virtual bool CanPlay() const override;
+	LEVELSEQUENCE_API virtual void OnStartedPlaying() override;
+	LEVELSEQUENCE_API virtual void OnStopped() override;
+	LEVELSEQUENCE_API virtual void UpdateMovieSceneInstance(FMovieSceneEvaluationRange InRange, EMovieScenePlayerStatus::Type PlayerStatus, const FMovieSceneUpdateArgs& Args) override;
 
 public:
 
 	/** Populate the specified array with any given event contexts for the specified world */
-	static void GetEventContexts(UWorld& InWorld, TArray<UObject*>& OutContexts);
+	static LEVELSEQUENCE_API void GetEventContexts(UWorld& InWorld, TArray<UObject*>& OutContexts);
 
 	/** Take a snapshot of the current state of this player */
-	void TakeFrameSnapshot(FLevelSequencePlayerSnapshot& OutSnapshot) const;
+	LEVELSEQUENCE_API void TakeFrameSnapshot(FLevelSequencePlayerSnapshot& OutSnapshot) const;
 
 	/** Set the offset time for the snapshot in play rate frames. */
 	void SetSnapshotOffsetFrames(int32 InFrameOffset) { SnapshotOffsetTime = TOptional<int32>(InFrameOffset); }
 
 private:
 
-	void EnableCinematicMode(bool bEnable);
+	LEVELSEQUENCE_API void EnableCinematicMode(bool bEnable);
 
 private:
 
diff --git a/Engine/Source/Runtime/LevelSequence/Public/LevelSequenceSpawnRegister.h b/Engine/Source/Runtime/LevelSequence/Public/LevelSequenceSpawnRegister.h
--- a/Engine/Source/Runtime/LevelSequence/Public/LevelSequenceSpawnRegister.h
+++ b/Engine/Source/Runtime/LevelSequence/Public/LevelSequenceSpawnRegister.h
@@ -16,18 +16,18 @@ class UObject;
 struct FMovieSceneSpawnable;
 
 /** Movie scene spawn register that knows how to handle spawning objects (actors) for a level sequence  */
-class LEVELSEQUENCE_API FLevelSequenceSpawnRegister : public FMovieSceneSpawnRegister
+class FLevelSequenceSpawnRegister : public FMovieSceneSpawnRegister
 {
 public:
-	FLevelSequenceSpawnRegister();
+	LEVELSEQUENCE_API FLevelSequenceSpawnRegister();
 
 protected:
 	/** ~ FMovieSceneSpawnRegister interface */
-	virtual UObject* SpawnObject(FMovieSceneSpawnable& Spawnable, FMovieSceneSequenceIDRef TemplateID, IMovieScenePlayer& Player) override;
-	virtual void DestroySpawnedObject(UObject& Object) override;
+	LEVELSEQUENCE_API virtual UObject* SpawnObject(FMovieSceneSpawnable& Spawnable, FMovieSceneSequenceIDRef TemplateID, IMovieScenePlayer& Player) override;
+	LEVELSEQUENCE_API virtual void DestroySpawnedObject(UObject& Object) override;
 
 #if WITH_EDITOR
-	virtual bool CanSpawnObject(UClass* InClass) const override;
+	LEVELSEQUENCE_API virtual bool CanSpawnObject(UClass* InClass) const override;
 #endif
 
 protected:
diff --git a/Engine/Source/Runtime/LevelSequence/Public/SequenceMediaController.h b/Engine/Source/Runtime/LevelSequence/Public/SequenceMediaController.h
--- a/Engine/Source/Runtime/LevelSequence/Public/SequenceMediaController.h
+++ b/Engine/Source/Runtime/LevelSequence/Public/SequenceMediaController.h
@@ -17,8 +17,8 @@ class UMediaPlayer;
 /**
  * Replicated actor class that is responsible for instigating various cinematic assets (Media, Audio, Level Sequences) in a synchronized fasion
  */
-UCLASS(hideCategories=(Rendering, Physics, HLOD, Activation, Input, Collision))
-class LEVELSEQUENCE_API ALevelSequenceMediaController
+UCLASS(hideCategories=(Rendering, Physics, HLOD, Activation, Input, Collision), MinimalAPI)
+class ALevelSequenceMediaController
 	: public AActor
 	, public IMovieSceneCustomClockSource
 {
@@ -26,10 +26,10 @@ public:
 
 	GENERATED_BODY()
 
-	ALevelSequenceMediaController(const FObjectInitializer& Init);
+	LEVELSEQUENCE_API ALevelSequenceMediaController(const FObjectInitializer& Init);
 
 	UFUNCTION(BlueprintCallable, Category="Synchronization")
-	void Play();
+	LEVELSEQUENCE_API void Play();
 
 	/**
 	 * Access this actor's media component
@@ -53,29 +53,29 @@ public:
 	 * Forcibly synchronize the sequence to the server's position if it has diverged by more than the specified threshold
 	 */
 	UFUNCTION(BlueprintCallable, Category="Synchronization")
-	void SynchronizeToServer(float DesyncThresholdSeconds = 2.f);
+	LEVELSEQUENCE_API void SynchronizeToServer(float DesyncThresholdSeconds = 2.f);
 
 private:
 
-	void Client_Play();
+	LEVELSEQUENCE_API void Client_Play();
 
-	void Client_ConditionallyForceTime(float DesyncThresholdSeconds);
+	LEVELSEQUENCE_API void Client_ConditionallyForceTime(float DesyncThresholdSeconds);
 
 private:
 
-	virtual void GetLifetimeReplicatedProps(TArray<FLifetimeProperty>& OutLifetimeProps) const;
+	LEVELSEQUENCE_API virtual void GetLifetimeReplicatedProps(TArray<FLifetimeProperty>& OutLifetimeProps) const;
 
 private:
 
-	virtual void OnTick(float DeltaSeconds, float InPlayRate) override final;
-	virtual void OnStartPlaying(const FQualifiedFrameTime& InStartTime) override final;
-	virtual void OnStopPlaying(const FQualifiedFrameTime& InStopTime) override final;
-	virtual FFrameTime OnRequestCurrentTime(const FQualifiedFrameTime& InCurrentTime, float InPlayRate) override final;
+	LEVELSEQUENCE_API virtual void OnTick(float DeltaSeconds, float InPlayRate) override final;
+	LEVELSEQUENCE_API virtual void OnStartPlaying(const FQualifiedFrameTime& InStartTime) override final;
+	LEVELSEQUENCE_API virtual void OnStopPlaying(const FQualifiedFrameTime& InStopTime) override final;
+	LEVELSEQUENCE_API virtual FFrameTime OnRequestCurrentTime(const FQualifiedFrameTime& InCurrentTime, float InPlayRate) override final;
 
 private:
 
 	UFUNCTION()
-	void OnRep_ServerStartTimeSeconds();
+	LEVELSEQUENCE_API void OnRep_ServerStartTimeSeconds();
 
 	/** Pointer to the sequence actor to use for playback */
 	UPROPERTY(EditAnywhere, BlueprintGetter=GetSequence, Category="Synchronization")
diff --git a/Engine/Source/Runtime/LiveLinkAnimationCore/Public/AnimNode_LiveLinkPose.h b/Engine/Source/Runtime/LiveLinkAnimationCore/Public/AnimNode_LiveLinkPose.h
--- a/Engine/Source/Runtime/LiveLinkAnimationCore/Public/AnimNode_LiveLinkPose.h
+++ b/Engine/Source/Runtime/LiveLinkAnimationCore/Public/AnimNode_LiveLinkPose.h
@@ -17,7 +17,7 @@ class ILiveLinkClient;
 PRAGMA_DISABLE_DEPRECATION_WARNINGS
 
 USTRUCT(BlueprintInternalUseOnly)
-struct LIVELINKANIMATIONCORE_API FAnimNode_LiveLinkPose : public FAnimNode_Base
+struct FAnimNode_LiveLinkPose : public FAnimNode_Base
 {
 	GENERATED_BODY()
 
@@ -41,22 +41,22 @@ public:
 	TObjectPtr<ULiveLinkRetargetAsset> CurrentRetargetAsset;
 
 public:
-	FAnimNode_LiveLinkPose();
+	LIVELINKANIMATIONCORE_API FAnimNode_LiveLinkPose();
 
 	//~ FAnimNode_Base interface
-	virtual void Initialize_AnyThread(const FAnimationInitializeContext& Context) override;
-	virtual void CacheBones_AnyThread(const FAnimationCacheBonesContext & Context) override;
-	virtual void Update_AnyThread(const FAnimationUpdateContext & Context) override;
-	virtual void Evaluate_AnyThread(FPoseContext& Output) override;
+	LIVELINKANIMATIONCORE_API virtual void Initialize_AnyThread(const FAnimationInitializeContext& Context) override;
+	LIVELINKANIMATIONCORE_API virtual void CacheBones_AnyThread(const FAnimationCacheBonesContext & Context) override;
+	LIVELINKANIMATIONCORE_API virtual void Update_AnyThread(const FAnimationUpdateContext & Context) override;
+	LIVELINKANIMATIONCORE_API virtual void Evaluate_AnyThread(FPoseContext& Output) override;
 	virtual bool HasPreUpdate() const { return true; }
-	virtual void PreUpdate(const UAnimInstance* InAnimInstance) override;
-	virtual void GatherDebugData(FNodeDebugData& DebugData) override;
+	LIVELINKANIMATIONCORE_API virtual void PreUpdate(const UAnimInstance* InAnimInstance) override;
+	LIVELINKANIMATIONCORE_API virtual void GatherDebugData(FNodeDebugData& DebugData) override;
 	//~ End of FAnimNode_Base interface
 
-	bool Serialize(FArchive& Ar);
+	LIVELINKANIMATIONCORE_API bool Serialize(FArchive& Ar);
 
 protected:
-	virtual void OnInitializeAnimInstance(const FAnimInstanceProxy* InProxy, const UAnimInstance* InAnimInstance) override;
+	LIVELINKANIMATIONCORE_API virtual void OnInitializeAnimInstance(const FAnimInstanceProxy* InProxy, const UAnimInstance* InAnimInstance) override;
 
 private:
 
diff --git a/Engine/Source/Runtime/LiveLinkAnimationCore/Public/LiveLinkInstance.h b/Engine/Source/Runtime/LiveLinkAnimationCore/Public/LiveLinkInstance.h
--- a/Engine/Source/Runtime/LiveLinkAnimationCore/Public/LiveLinkInstance.h
+++ b/Engine/Source/Runtime/LiveLinkAnimationCore/Public/LiveLinkInstance.h
@@ -12,7 +12,7 @@ class ULiveLinkRetargetAsset;
 
 /** Proxy override for this UAnimInstance-derived class */
 USTRUCT()
-struct LIVELINKANIMATIONCORE_API FLiveLinkInstanceProxy : public FAnimInstanceProxy
+struct FLiveLinkInstanceProxy : public FAnimInstanceProxy
 {
 public:
 	friend struct FAnimNode_LiveLinkPose;
@@ -28,17 +28,17 @@ public:
 	{
 	}
 
-	virtual void Initialize(UAnimInstance* InAnimInstance) override;
-	virtual void PreUpdate(UAnimInstance* InAnimInstance, float DeltaSeconds) override;
-	virtual bool Evaluate(FPoseContext& Output) override;
-	virtual void UpdateAnimationNode(const FAnimationUpdateContext& InContext) override;
+	LIVELINKANIMATIONCORE_API virtual void Initialize(UAnimInstance* InAnimInstance) override;
+	LIVELINKANIMATIONCORE_API virtual void PreUpdate(UAnimInstance* InAnimInstance, float DeltaSeconds) override;
+	LIVELINKANIMATIONCORE_API virtual bool Evaluate(FPoseContext& Output) override;
+	LIVELINKANIMATIONCORE_API virtual void UpdateAnimationNode(const FAnimationUpdateContext& InContext) override;
 
 	UPROPERTY(EditAnywhere, Category = Settings)
 	FAnimNode_LiveLinkPose PoseNode;
 };
 
-UCLASS(transient, NotBlueprintable)
-class LIVELINKANIMATIONCORE_API ULiveLinkInstance : public UAnimInstance
+UCLASS(transient, NotBlueprintable, MinimalAPI)
+class ULiveLinkInstance : public UAnimInstance
 {
 	GENERATED_BODY()
 
@@ -57,8 +57,8 @@ public:
 	}
 
 protected:
-	virtual FAnimInstanceProxy* CreateAnimInstanceProxy() override;
-	virtual void DestroyAnimInstanceProxy(FAnimInstanceProxy* InProxy) override;
+	LIVELINKANIMATIONCORE_API virtual FAnimInstanceProxy* CreateAnimInstanceProxy() override;
+	LIVELINKANIMATIONCORE_API virtual void DestroyAnimInstanceProxy(FAnimInstanceProxy* InProxy) override;
 
 	// Cache for GC
 	UPROPERTY(transient)
diff --git a/Engine/Source/Runtime/LiveLinkAnimationCore/Public/LiveLinkRemapAsset.h b/Engine/Source/Runtime/LiveLinkAnimationCore/Public/LiveLinkRemapAsset.h
--- a/Engine/Source/Runtime/LiveLinkAnimationCore/Public/LiveLinkRemapAsset.h
+++ b/Engine/Source/Runtime/LiveLinkAnimationCore/Public/LiveLinkRemapAsset.h
@@ -22,37 +22,37 @@ struct FLiveLinkBaseStaticData;
 struct FLiveLinkSkeletonStaticData;
 
 // Remap asset for data coming from Live Link. Allows simple application of bone transforms into current pose based on name remapping only
-UCLASS(Blueprintable)
-class LIVELINKANIMATIONCORE_API ULiveLinkRemapAsset : public ULiveLinkRetargetAsset
+UCLASS(Blueprintable, MinimalAPI)
+class ULiveLinkRemapAsset : public ULiveLinkRetargetAsset
 {
 	GENERATED_UCLASS_BODY()
 
 	virtual ~ULiveLinkRemapAsset() {}
 
 	//~ Begin UObject Interface
-	virtual void BeginDestroy() override;
+	LIVELINKANIMATIONCORE_API virtual void BeginDestroy() override;
 	//~ End UObject Interface
 
 	//~ Begin ULiveLinkRetargetAsset interface
-	virtual void BuildPoseFromAnimationData(float DeltaTime, const FLiveLinkSkeletonStaticData* InSkeletonData, const FLiveLinkAnimationFrameData* InFrameData, FCompactPose& OutPose) override;
-	virtual void BuildPoseAndCurveFromBaseData(float DeltaTime, const FLiveLinkBaseStaticData* InBaseStaticData, const FLiveLinkBaseFrameData* InBaseFrameData, FCompactPose& OutPose, FBlendedCurve& OutCurve) override;
+	LIVELINKANIMATIONCORE_API virtual void BuildPoseFromAnimationData(float DeltaTime, const FLiveLinkSkeletonStaticData* InSkeletonData, const FLiveLinkAnimationFrameData* InFrameData, FCompactPose& OutPose) override;
+	LIVELINKANIMATIONCORE_API virtual void BuildPoseAndCurveFromBaseData(float DeltaTime, const FLiveLinkBaseStaticData* InBaseStaticData, const FLiveLinkBaseFrameData* InBaseFrameData, FCompactPose& OutPose, FBlendedCurve& OutCurve) override;
 	//~ End ULiveLinkRetargetAsset interface
 
 	/** Blueprint Implementable function for getting a remapped bone name from the original */
 	UFUNCTION(BlueprintNativeEvent, Category = "Live Link Remap")
-	FName GetRemappedBoneName(FName BoneName) const;
+	LIVELINKANIMATIONCORE_API FName GetRemappedBoneName(FName BoneName) const;
 
 	/** Blueprint Implementable function for getting a remapped curve name from the original */
 	UFUNCTION(BlueprintNativeEvent, Category = "Live Link Remap")
-	FName GetRemappedCurveName(FName CurveName) const;
+	LIVELINKANIMATIONCORE_API FName GetRemappedCurveName(FName CurveName) const;
 
 	/** Blueprint Implementable function for remapping, adding or otherwise modifying the curve element data from Live Link. This is run after GetRemappedCurveName */
 	UFUNCTION(BlueprintNativeEvent, Category = "Live Link Remap")
-	void RemapCurveElements(UPARAM(ref)TMap<FName, float>& CurveItems) const;
+	LIVELINKANIMATIONCORE_API void RemapCurveElements(UPARAM(ref)TMap<FName, float>& CurveItems) const;
 
 private:
 
-	void OnBlueprintClassCompiled(UBlueprint* TargetBlueprint);
+	LIVELINKANIMATIONCORE_API void OnBlueprintClassCompiled(UBlueprint* TargetBlueprint);
 
 	// Name mapping between source bone name and transformed bone name
 	// (returned from GetRemappedBoneName)
@@ -66,4 +66,4 @@ private:
 	/** Blueprint.OnCompiled delegate handle */
 	FDelegateHandle OnBlueprintCompiledDelegate;
 #endif
-};
\ No newline at end of file
+};
diff --git a/Engine/Source/Runtime/LiveLinkAnimationCore/Public/LiveLinkRetargetAsset.h b/Engine/Source/Runtime/LiveLinkAnimationCore/Public/LiveLinkRetargetAsset.h
--- a/Engine/Source/Runtime/LiveLinkAnimationCore/Public/LiveLinkRetargetAsset.h
+++ b/Engine/Source/Runtime/LiveLinkAnimationCore/Public/LiveLinkRetargetAsset.h
@@ -20,19 +20,19 @@ struct FLiveLinkBaseStaticData;
 struct FLiveLinkSkeletonStaticData;
 
 // Base class for retargeting live link data. 
-UCLASS(Abstract)
-class LIVELINKANIMATIONCORE_API ULiveLinkRetargetAsset : public UObject
+UCLASS(Abstract, MinimalAPI)
+class ULiveLinkRetargetAsset : public UObject
 {
 	GENERATED_UCLASS_BODY()
 
 	// Takes the supplied curve name and value and applies it to the blended curve (as approriate given the supplied skeleton
-	void ApplyCurveValue(const USkeleton* Skeleton, const FName CurveName, const float CurveValue, FBlendedCurve& OutCurve) const;
+	LIVELINKANIMATIONCORE_API void ApplyCurveValue(const USkeleton* Skeleton, const FName CurveName, const float CurveValue, FBlendedCurve& OutCurve) const;
 
 	// Builds curve data into OutCurve from the supplied live link frame
-	void BuildCurveData(const FLiveLinkSkeletonStaticData* InSkeletonData, const FLiveLinkAnimationFrameData* InFrameData, const FCompactPose& InPose, FBlendedCurve& OutCurve) const;
+	LIVELINKANIMATIONCORE_API void BuildCurveData(const FLiveLinkSkeletonStaticData* InSkeletonData, const FLiveLinkAnimationFrameData* InFrameData, const FCompactPose& InPose, FBlendedCurve& OutCurve) const;
 
 	// Builds curve data into OutCurve from the supplied map of curve name to float
-	void BuildCurveData(const TMap<FName, float>& CurveMap, const FCompactPose& InPose, FBlendedCurve& OutCurve) const;
+	LIVELINKANIMATIONCORE_API void BuildCurveData(const TMap<FName, float>& CurveMap, const FCompactPose& InPose, FBlendedCurve& OutCurve) const;
 
 	// Called once when the retargeter is created 
 	virtual void Initialize() {}
diff --git a/Engine/Source/Runtime/LiveLinkInterface/Public/ILiveLinkSubject.h b/Engine/Source/Runtime/LiveLinkInterface/Public/ILiveLinkSubject.h
--- a/Engine/Source/Runtime/LiveLinkInterface/Public/ILiveLinkSubject.h
+++ b/Engine/Source/Runtime/LiveLinkInterface/Public/ILiveLinkSubject.h
@@ -20,18 +20,18 @@ class ULiveLinkRole;
  * Subjects are individual streams of data within the client.
  * An animating character could be a subject for instance.
  */
-class LIVELINKINTERFACE_API ILiveLinkSubject
+class ILiveLinkSubject
 {
 public:
 	virtual ~ILiveLinkSubject() {}
 	virtual void Initialize(FLiveLinkSubjectKey SubjectKey, TSubclassOf<ULiveLinkRole> Role, ILiveLinkClient* LiveLinkClient) = 0;
 	virtual void Update() = 0;
-	virtual bool EvaluateFrame(TSubclassOf<ULiveLinkRole> InDesiredRole, FLiveLinkSubjectFrameData& OutFrame);
+	LIVELINKINTERFACE_API virtual bool EvaluateFrame(TSubclassOf<ULiveLinkRole> InDesiredRole, FLiveLinkSubjectFrameData& OutFrame);
 	virtual void ClearFrames() = 0;
 
 	virtual FLiveLinkSubjectKey GetSubjectKey() const = 0;
 	virtual TSubclassOf<ULiveLinkRole> GetRole() const = 0;
-	virtual bool SupportsRole(TSubclassOf<ULiveLinkRole> InDesiredRole) const;
+	LIVELINKINTERFACE_API virtual bool SupportsRole(TSubclassOf<ULiveLinkRole> InDesiredRole) const;
 
 	virtual bool HasValidFrameSnapshot() const = 0;
 	virtual FLiveLinkStaticDataStruct& GetStaticData() = 0;
@@ -55,5 +55,5 @@ protected:
 	virtual const FLiveLinkSubjectFrameData& GetFrameSnapshot() const = 0;
 
 protected:
-	static bool Translate(const ILiveLinkSubject* LinkSubject, TSubclassOf<ULiveLinkRole> DesiredRole, const FLiveLinkStaticDataStruct& StaticData, const FLiveLinkFrameDataStruct& FrameData, FLiveLinkSubjectFrameData& OutFrame);
-};
\ No newline at end of file
+	static LIVELINKINTERFACE_API bool Translate(const ILiveLinkSubject* LinkSubject, TSubclassOf<ULiveLinkRole> DesiredRole, const FLiveLinkStaticDataStruct& StaticData, const FLiveLinkFrameDataStruct& FrameData, FLiveLinkSubjectFrameData& OutFrame);
+};
diff --git a/Engine/Source/Runtime/LiveLinkInterface/Public/LiveLinkController.h b/Engine/Source/Runtime/LiveLinkInterface/Public/LiveLinkController.h
--- a/Engine/Source/Runtime/LiveLinkInterface/Public/LiveLinkController.h
+++ b/Engine/Source/Runtime/LiveLinkInterface/Public/LiveLinkController.h
@@ -17,13 +17,13 @@ class AActor;
 /**
  * Basic object to control a UObject live link frames
  */
-UCLASS(Abstract)
-class LIVELINKINTERFACE_API ULiveLinkController : public UObject
+UCLASS(Abstract, MinimalAPI)
+class ULiveLinkController : public UObject
 {
 	GENERATED_BODY()
 
 public:
-	virtual TSubclassOf<ULiveLinkRole> GetRole() const PURE_VIRTUAL(ULiveLinkController::GetRole, return TSubclassOf<ULiveLinkRole>(););
+	LIVELINKINTERFACE_API virtual TSubclassOf<ULiveLinkRole> GetRole() const PURE_VIRTUAL(ULiveLinkController::GetRole, return TSubclassOf<ULiveLinkRole>(););
 
 	/** Initialize the controller at the first tick of its owner component. */
 	virtual void OnRegistered() { }
diff --git a/Engine/Source/Runtime/LiveLinkInterface/Public/LiveLinkCurveRemapSettings.h b/Engine/Source/Runtime/LiveLinkInterface/Public/LiveLinkCurveRemapSettings.h
--- a/Engine/Source/Runtime/LiveLinkInterface/Public/LiveLinkCurveRemapSettings.h
+++ b/Engine/Source/Runtime/LiveLinkInterface/Public/LiveLinkCurveRemapSettings.h
@@ -24,8 +24,8 @@ struct FLiveLinkCurveConversionSettings
 	TMap<FString, FSoftObjectPath> CurveConversionAssetMap;
 };
 
-UCLASS(config=Engine, defaultconfig, meta=(DisplayName="LiveLink"))
-class LIVELINKINTERFACE_API ULiveLinkCurveRemapSettings : public ULiveLinkSourceSettings
+UCLASS(config=Engine, defaultconfig, meta=(DisplayName="LiveLink"), MinimalAPI)
+class ULiveLinkCurveRemapSettings : public ULiveLinkSourceSettings
 {
 public:
 	GENERATED_BODY()
@@ -36,7 +36,7 @@ public:
 #if WITH_EDITOR
 
 	//UObject override so we can change this setting when changed in editor
-	virtual void PostEditChangeProperty(struct FPropertyChangedEvent& PropertyChangedEvent) override;
+	LIVELINKINTERFACE_API virtual void PostEditChangeProperty(struct FPropertyChangedEvent& PropertyChangedEvent) override;
 
 #endif // WITH_EDITOR
 };
diff --git a/Engine/Source/Runtime/LiveLinkInterface/Public/LiveLinkFrameInterpolationProcessor.h b/Engine/Source/Runtime/LiveLinkInterface/Public/LiveLinkFrameInterpolationProcessor.h
--- a/Engine/Source/Runtime/LiveLinkInterface/Public/LiveLinkFrameInterpolationProcessor.h
+++ b/Engine/Source/Runtime/LiveLinkInterface/Public/LiveLinkFrameInterpolationProcessor.h
@@ -42,7 +42,7 @@ struct FLiveLinkInterpolationInfo
  * Inherit from it to do custom blending for your data type
  * @note It can be called from any thread
  */
-class LIVELINKINTERFACE_API ILiveLinkFrameInterpolationProcessorWorker
+class ILiveLinkFrameInterpolationProcessorWorker
 {
 public:
 	virtual TSubclassOf<ULiveLinkRole> GetRole() const = 0;
@@ -57,14 +57,14 @@ public:
  * Inherit from it to do custom blending for your data type
  * @note It can only be used on the Game Thread. See ILiveLinkFrameInterpolationProcessorWorker for the any thread implementation.
  */
-UCLASS(Abstract, editinlinenew, ClassGroup = (LiveLink))
-class LIVELINKINTERFACE_API ULiveLinkFrameInterpolationProcessor : public UObject
+UCLASS(Abstract, editinlinenew, ClassGroup = (LiveLink), MinimalAPI)
+class ULiveLinkFrameInterpolationProcessor : public UObject
 {
 	GENERATED_BODY()
 
 public:
 	using FWorkerSharedPtr = TSharedPtr<ILiveLinkFrameInterpolationProcessorWorker, ESPMode::ThreadSafe>;
 
-	virtual TSubclassOf<ULiveLinkRole> GetRole() const PURE_VIRTUAL(ULiveLinkFrameInterpolationProcessor::GetFromRole, return TSubclassOf<ULiveLinkRole>(););
-	virtual FWorkerSharedPtr FetchWorker() PURE_VIRTUAL(ULiveLinkFrameInterpolationProcessor::FetchWorker, return FWorkerSharedPtr(););
+	LIVELINKINTERFACE_API virtual TSubclassOf<ULiveLinkRole> GetRole() const PURE_VIRTUAL(ULiveLinkFrameInterpolationProcessor::GetFromRole, return TSubclassOf<ULiveLinkRole>(););
+	LIVELINKINTERFACE_API virtual FWorkerSharedPtr FetchWorker() PURE_VIRTUAL(ULiveLinkFrameInterpolationProcessor::FetchWorker, return FWorkerSharedPtr(););
 };
diff --git a/Engine/Source/Runtime/LiveLinkInterface/Public/LiveLinkFramePreProcessor.h b/Engine/Source/Runtime/LiveLinkInterface/Public/LiveLinkFramePreProcessor.h
--- a/Engine/Source/Runtime/LiveLinkInterface/Public/LiveLinkFramePreProcessor.h
+++ b/Engine/Source/Runtime/LiveLinkInterface/Public/LiveLinkFramePreProcessor.h
@@ -18,7 +18,7 @@
  * Inherit from it to add specific operations / options for a certain type of data
  * @note It can be called from any thread
  */
-class LIVELINKINTERFACE_API ILiveLinkFramePreProcessorWorker
+class ILiveLinkFramePreProcessorWorker
 {
 public:
 	virtual TSubclassOf<ULiveLinkRole> GetRole() const = 0;
@@ -31,14 +31,14 @@ public:
  * Inherit from it to add specific operations / options for a certain type of data
  * @note It can only be used on the Game Thread. See ILiveLinkFramePreProcessorWorker for the any thread implementation.
  */
-UCLASS(Abstract, Blueprintable, editinlinenew, ClassGroup = (LiveLink))
-class LIVELINKINTERFACE_API ULiveLinkFramePreProcessor : public UObject
+UCLASS(Abstract, Blueprintable, editinlinenew, ClassGroup = (LiveLink), MinimalAPI)
+class ULiveLinkFramePreProcessor : public UObject
 {
 	GENERATED_BODY()
 
 public:
 	using FWorkerSharedPtr = TSharedPtr<ILiveLinkFramePreProcessorWorker, ESPMode::ThreadSafe>;
 
-	virtual TSubclassOf<ULiveLinkRole> GetRole() const PURE_VIRTUAL(ULiveLinkFramePreProcessor::GetFromRole, return TSubclassOf<ULiveLinkRole>(););
-	virtual FWorkerSharedPtr FetchWorker() PURE_VIRTUAL(ULiveLinkFramePreProcessor::FetchWorker, return FWorkerSharedPtr(););
-};
\ No newline at end of file
+	LIVELINKINTERFACE_API virtual TSubclassOf<ULiveLinkRole> GetRole() const PURE_VIRTUAL(ULiveLinkFramePreProcessor::GetFromRole, return TSubclassOf<ULiveLinkRole>(););
+	LIVELINKINTERFACE_API virtual FWorkerSharedPtr FetchWorker() PURE_VIRTUAL(ULiveLinkFramePreProcessor::FetchWorker, return FWorkerSharedPtr(););
+};
diff --git a/Engine/Source/Runtime/LiveLinkInterface/Public/LiveLinkFrameTranslator.h b/Engine/Source/Runtime/LiveLinkInterface/Public/LiveLinkFrameTranslator.h
--- a/Engine/Source/Runtime/LiveLinkInterface/Public/LiveLinkFrameTranslator.h
+++ b/Engine/Source/Runtime/LiveLinkInterface/Public/LiveLinkFrameTranslator.h
@@ -22,10 +22,10 @@ class ULiveLinkRole;
  * Basic object to translate data from one role to another
  * @note It can be called from any thread
  */
-class LIVELINKINTERFACE_API ILiveLinkFrameTranslatorWorker
+class ILiveLinkFrameTranslatorWorker
 {
 public:
-	bool CanTranslate(TSubclassOf<ULiveLinkRole> ToRole) const;
+	LIVELINKINTERFACE_API bool CanTranslate(TSubclassOf<ULiveLinkRole> ToRole) const;
 	virtual TSubclassOf<ULiveLinkRole> GetFromRole() const = 0;
 	virtual TSubclassOf<ULiveLinkRole> GetToRole() const = 0;
 	virtual bool Translate(const FLiveLinkStaticDataStruct& InStaticData, const FLiveLinkFrameDataStruct& InFrameData, FLiveLinkSubjectFrameData& OutTranslatedFrame) const = 0;
@@ -36,17 +36,17 @@ public:
  * Basic object to translate data from one role to another
  * @note It can only be used on the Game Thread. See ILiveLinkFrameTranslatorWorker for the any thread implementation.
  */
-UCLASS(Abstract, editinlinenew, ClassGroup = (LiveLink))
-class LIVELINKINTERFACE_API ULiveLinkFrameTranslator : public UObject
+UCLASS(Abstract, editinlinenew, ClassGroup = (LiveLink), MinimalAPI)
+class ULiveLinkFrameTranslator : public UObject
 {
 	GENERATED_BODY()
 
 public:
 	using FWorkerSharedPtr = TSharedPtr<ILiveLinkFrameTranslatorWorker, ESPMode::ThreadSafe>;
 
-	bool CanTranslate(TSubclassOf<ULiveLinkRole> ToRole) const;
+	LIVELINKINTERFACE_API bool CanTranslate(TSubclassOf<ULiveLinkRole> ToRole) const;
 
-	virtual TSubclassOf<ULiveLinkRole> GetFromRole() const PURE_VIRTUAL(ULiveLinkFrameTranslator::GetFromRole, return TSubclassOf<ULiveLinkRole>(););
-	virtual TSubclassOf<ULiveLinkRole> GetToRole() const PURE_VIRTUAL(ULiveLinkFrameTranslator::GetToRole, return TSubclassOf<ULiveLinkRole>(););
-	virtual FWorkerSharedPtr FetchWorker() PURE_VIRTUAL(ULiveLinkFrameTranslator::FetchWorker, return FWorkerSharedPtr(););
+	LIVELINKINTERFACE_API virtual TSubclassOf<ULiveLinkRole> GetFromRole() const PURE_VIRTUAL(ULiveLinkFrameTranslator::GetFromRole, return TSubclassOf<ULiveLinkRole>(););
+	LIVELINKINTERFACE_API virtual TSubclassOf<ULiveLinkRole> GetToRole() const PURE_VIRTUAL(ULiveLinkFrameTranslator::GetToRole, return TSubclassOf<ULiveLinkRole>(););
+	LIVELINKINTERFACE_API virtual FWorkerSharedPtr FetchWorker() PURE_VIRTUAL(ULiveLinkFrameTranslator::FetchWorker, return FWorkerSharedPtr(););
 };
diff --git a/Engine/Source/Runtime/LiveLinkInterface/Public/LiveLinkLog.h b/Engine/Source/Runtime/LiveLinkInterface/Public/LiveLinkLog.h
--- a/Engine/Source/Runtime/LiveLinkInterface/Public/LiveLinkLog.h
+++ b/Engine/Source/Runtime/LiveLinkInterface/Public/LiveLinkLog.h
@@ -20,7 +20,7 @@ struct FTimespan;
 #define ENABLE_LIVELINK_LOGGING (1 && !NO_LOGGING && !(UE_BUILD_SHIPPING || UE_BUILD_TEST))
 
 /** This class represents a log of LiveLink output each of which can be a rich tokenized message */
-class LIVELINKINTERFACE_API FLiveLinkLog
+class FLiveLinkLog
 {
 public:
 	/** Write an error in to the LiveLink log. */
@@ -148,5 +148,5 @@ protected:
 
 protected:
 	/** The instance that will manage the logging */
-	static TUniquePtr<FLiveLinkLog> Instance;
+	static LIVELINKINTERFACE_API TUniquePtr<FLiveLinkLog> Instance;
 };
diff --git a/Engine/Source/Runtime/LiveLinkInterface/Public/LiveLinkRole.h b/Engine/Source/Runtime/LiveLinkInterface/Public/LiveLinkRole.h
--- a/Engine/Source/Runtime/LiveLinkInterface/Public/LiveLinkRole.h
+++ b/Engine/Source/Runtime/LiveLinkInterface/Public/LiveLinkRole.h
@@ -13,26 +13,26 @@
 /**
  * Basic object to extend the meaning of incoming live link frames.
  */
-UCLASS(Abstract)
-class LIVELINKINTERFACE_API ULiveLinkRole : public UObject
+UCLASS(Abstract, MinimalAPI)
+class ULiveLinkRole : public UObject
 {
 	GENERATED_BODY()
 
 public:
-	virtual UScriptStruct* GetStaticDataStruct() const PURE_VIRTUAL(ULiveLinkRole::GetStaticDataStruct, return nullptr;);
-	virtual UScriptStruct* GetFrameDataStruct() const PURE_VIRTUAL(ULiveLinkRole::GetFrameDataStruct, return nullptr;);
-	virtual UScriptStruct* GetBlueprintDataStruct() const PURE_VIRTUAL(ULiveLinkRole::GetBlueprintDataStruct, return nullptr;);
+	LIVELINKINTERFACE_API virtual UScriptStruct* GetStaticDataStruct() const PURE_VIRTUAL(ULiveLinkRole::GetStaticDataStruct, return nullptr;);
+	LIVELINKINTERFACE_API virtual UScriptStruct* GetFrameDataStruct() const PURE_VIRTUAL(ULiveLinkRole::GetFrameDataStruct, return nullptr;);
+	LIVELINKINTERFACE_API virtual UScriptStruct* GetBlueprintDataStruct() const PURE_VIRTUAL(ULiveLinkRole::GetBlueprintDataStruct, return nullptr;);
 
-	virtual bool InitializeBlueprintData(const FLiveLinkSubjectFrameData& InSourceData, FLiveLinkBlueprintDataStruct& OutBlueprintData) const PURE_VIRTUAL(ULiveLinkRole::InitializeBlueprintData, return false;);
+	LIVELINKINTERFACE_API virtual bool InitializeBlueprintData(const FLiveLinkSubjectFrameData& InSourceData, FLiveLinkBlueprintDataStruct& OutBlueprintData) const PURE_VIRTUAL(ULiveLinkRole::InitializeBlueprintData, return false;);
 
-	virtual FText GetDisplayName() const;
+	LIVELINKINTERFACE_API virtual FText GetDisplayName() const;
 	virtual bool IsStaticDataValid(const FLiveLinkStaticDataStruct& InStaticData, bool& bOutShouldLogWarning) const { return true; }
 	virtual bool IsFrameDataValid(const FLiveLinkStaticDataStruct& InStaticData, const FLiveLinkFrameDataStruct& InFrameData, bool& bOutShouldLogWarning) const { return true; }
 };
 
 
 USTRUCT(BlueprintType)
-struct LIVELINKINTERFACE_API FLiveLinkSubjectRepresentation
+struct FLiveLinkSubjectRepresentation
 {
 	GENERATED_BODY()
 
diff --git a/Engine/Source/Runtime/LiveLinkInterface/Public/LiveLinkRoleTrait.h b/Engine/Source/Runtime/LiveLinkInterface/Public/LiveLinkRoleTrait.h
--- a/Engine/Source/Runtime/LiveLinkInterface/Public/LiveLinkRoleTrait.h
+++ b/Engine/Source/Runtime/LiveLinkInterface/Public/LiveLinkRoleTrait.h
@@ -21,34 +21,34 @@ class ULiveLinkRole;
 class ULiveLinkVirtualSubject;
 
 
-class LIVELINKINTERFACE_API FLiveLinkRoleTrait
+class FLiveLinkRoleTrait
 {
 public:
 	/** Get the list of all the LiveLinkRole. */
-	static TArray<TSubclassOf<ULiveLinkRole>> GetRoles();
+	static LIVELINKINTERFACE_API TArray<TSubclassOf<ULiveLinkRole>> GetRoles();
 
 	/** Get the list of all the Frame Interpolation Processor for the Role. */
-	static TArray<TSubclassOf<ULiveLinkFrameInterpolationProcessor>> GetFrameInterpolationProcessorClasses(TSubclassOf<ULiveLinkRole> Role);
+	static LIVELINKINTERFACE_API TArray<TSubclassOf<ULiveLinkFrameInterpolationProcessor>> GetFrameInterpolationProcessorClasses(TSubclassOf<ULiveLinkRole> Role);
 
 	/** Get the list of all the Frame Pre Processor for the Role. */
-	static TArray<TSubclassOf<ULiveLinkFramePreProcessor>> GetFramePreProcessorClasses(TSubclassOf<ULiveLinkRole> Role);
+	static LIVELINKINTERFACE_API TArray<TSubclassOf<ULiveLinkFramePreProcessor>> GetFramePreProcessorClasses(TSubclassOf<ULiveLinkRole> Role);
 
 	/** Get the list of all the Frame Translator to convert into a Role. */
-	static TArray<TSubclassOf<ULiveLinkFrameTranslator>> GetFrameTranslatorClassesTo(TSubclassOf<ULiveLinkRole> Role);
+	static LIVELINKINTERFACE_API TArray<TSubclassOf<ULiveLinkFrameTranslator>> GetFrameTranslatorClassesTo(TSubclassOf<ULiveLinkRole> Role);
 
 	/** Get the list of all the Frame Translator to convert from a Role. */
-	static TArray<TSubclassOf<ULiveLinkFrameTranslator>> GetFrameTranslatorClassesFrom(TSubclassOf<ULiveLinkRole> Role);
+	static LIVELINKINTERFACE_API TArray<TSubclassOf<ULiveLinkFrameTranslator>> GetFrameTranslatorClassesFrom(TSubclassOf<ULiveLinkRole> Role);
 
 	/** Get the list of all the Frame Virtual Subject for a Role. */
-	static TArray<TSubclassOf<ULiveLinkVirtualSubject>> GetVirtualSubjectClasses();
-	static TArray<TSubclassOf<ULiveLinkVirtualSubject>> GetVirtualSubjectClasses(TSubclassOf<ULiveLinkRole> Role);
+	static LIVELINKINTERFACE_API TArray<TSubclassOf<ULiveLinkVirtualSubject>> GetVirtualSubjectClasses();
+	static LIVELINKINTERFACE_API TArray<TSubclassOf<ULiveLinkVirtualSubject>> GetVirtualSubjectClasses(TSubclassOf<ULiveLinkRole> Role);
 
 	/** Get the Controller class for a Role. */
-	static TSubclassOf<ULiveLinkController> GetControllerClass(TSubclassOf<ULiveLinkRole> Role);
+	static LIVELINKINTERFACE_API TSubclassOf<ULiveLinkController> GetControllerClass(TSubclassOf<ULiveLinkRole> Role);
 
 	/** Return true if the StaticData matches the Role static data type. */
-	static bool Validate(TSubclassOf<ULiveLinkRole> Role, const FLiveLinkStaticDataStruct& StaticData);
+	static LIVELINKINTERFACE_API bool Validate(TSubclassOf<ULiveLinkRole> Role, const FLiveLinkStaticDataStruct& StaticData);
 
 	/** Return true if the FrameData matches the Role frame data type. */
-	static bool Validate(TSubclassOf<ULiveLinkRole> Role, const FLiveLinkFrameDataStruct& FrameData);
+	static LIVELINKINTERFACE_API bool Validate(TSubclassOf<ULiveLinkRole> Role, const FLiveLinkFrameDataStruct& FrameData);
 };
diff --git a/Engine/Source/Runtime/LiveLinkInterface/Public/LiveLinkSourceFactory.h b/Engine/Source/Runtime/LiveLinkInterface/Public/LiveLinkSourceFactory.h
--- a/Engine/Source/Runtime/LiveLinkInterface/Public/LiveLinkSourceFactory.h
+++ b/Engine/Source/Runtime/LiveLinkInterface/Public/LiveLinkSourceFactory.h
@@ -22,23 +22,23 @@ class SWidget;
 /**
  * Base class of factory that creates Source. A source can be created in the editor via the Live Link panel or at runtime via a connection string.
  */
-UCLASS(Abstract)
-class LIVELINKINTERFACE_API ULiveLinkSourceFactory : public UObject
+UCLASS(Abstract, MinimalAPI)
+class ULiveLinkSourceFactory : public UObject
 {
 	GENERATED_BODY()
 
 public:
 	/** The name of the menu item (of any EMenuType) */
-	virtual FText GetSourceDisplayName() const PURE_VIRTUAL( ULiveLinkSourceFactory::GetSourceDisplayName, return FText(); );
+	LIVELINKINTERFACE_API virtual FText GetSourceDisplayName() const PURE_VIRTUAL( ULiveLinkSourceFactory::GetSourceDisplayName, return FText(); );
 
 	/** The tooltip of the menu item (of any EMenyType) */
-	virtual FText GetSourceTooltip() const PURE_VIRTUAL(ULiveLinkSourceFactory::GetSourceTooltip, return FText(); );
+	LIVELINKINTERFACE_API virtual FText GetSourceTooltip() const PURE_VIRTUAL(ULiveLinkSourceFactory::GetSourceTooltip, return FText(); );
 
 	UE_DEPRECATED(4.23, "CreateSourceCreationPanel is deprecated. LiveLinkSourceFactory can now be used at runtime. This factory won't work until it's been updated.")
-	virtual TSharedPtr<SWidget> CreateSourceCreationPanel();
+	LIVELINKINTERFACE_API virtual TSharedPtr<SWidget> CreateSourceCreationPanel();
 
 	UE_DEPRECATED(4.23, "OnSourceCreationPanelClosed is deprecated. LiveLinkSourceFactory can now be used at runtime. The factory won't work until it's been updated.")
-	virtual TSharedPtr<ILiveLinkSource> OnSourceCreationPanelClosed(bool bMakeSource);
+	LIVELINKINTERFACE_API virtual TSharedPtr<ILiveLinkSource> OnSourceCreationPanelClosed(bool bMakeSource);
 
 	enum class EMenuType
 	{
@@ -60,8 +60,8 @@ public:
 	 * @param OnLiveLinkSourceCreated Callback to call when the source has been created by the custom UI.
 	 * @return The subpanel UI created by the factory.
 	 */
-	virtual TSharedPtr<SWidget> BuildCreationPanel(FOnLiveLinkSourceCreated OnLiveLinkSourceCreated) const;
+	LIVELINKINTERFACE_API virtual TSharedPtr<SWidget> BuildCreationPanel(FOnLiveLinkSourceCreated OnLiveLinkSourceCreated) const;
 
 	/** Create a new source from a ConnectionString */
-	virtual TSharedPtr<ILiveLinkSource> CreateSource(const FString& ConnectionString) const PURE_VIRTUAL(ULiveLinkSourceFactory::CreateSource, return TSharedPtr<ILiveLinkSource>(); );
+	LIVELINKINTERFACE_API virtual TSharedPtr<ILiveLinkSource> CreateSource(const FString& ConnectionString) const PURE_VIRTUAL(ULiveLinkSourceFactory::CreateSource, return TSharedPtr<ILiveLinkSource>(); );
 };
diff --git a/Engine/Source/Runtime/LiveLinkInterface/Public/LiveLinkSourceSettings.h b/Engine/Source/Runtime/LiveLinkInterface/Public/LiveLinkSourceSettings.h
--- a/Engine/Source/Runtime/LiveLinkInterface/Public/LiveLinkSourceSettings.h
+++ b/Engine/Source/Runtime/LiveLinkInterface/Public/LiveLinkSourceSettings.h
@@ -142,8 +142,8 @@ struct FLiveLinkSourceDebugInfo
 };
 
 /** Base class for live link source settings (can be replaced by sources themselves) */
-UCLASS()
-class LIVELINKINTERFACE_API ULiveLinkSourceSettings : public UObject
+UCLASS(MinimalAPI)
+class ULiveLinkSourceSettings : public UObject
 {
 public:
 	GENERATED_BODY()
@@ -172,10 +172,10 @@ public:
 	TArray<FLiveLinkSourceDebugInfo> SourceDebugInfos_DEPRECATED;
 #endif
 
-	virtual void Serialize(FArchive& Ar) override;
+	LIVELINKINTERFACE_API virtual void Serialize(FArchive& Ar) override;
 
 #if WITH_EDITOR
-	virtual bool CanEditChange(const FProperty* InProperty) const override;
+	LIVELINKINTERFACE_API virtual bool CanEditChange(const FProperty* InProperty) const override;
 #endif
 };
 
diff --git a/Engine/Source/Runtime/LiveLinkInterface/Public/LiveLinkSubjectSettings.h b/Engine/Source/Runtime/LiveLinkInterface/Public/LiveLinkSubjectSettings.h
--- a/Engine/Source/Runtime/LiveLinkInterface/Public/LiveLinkSubjectSettings.h
+++ b/Engine/Source/Runtime/LiveLinkInterface/Public/LiveLinkSubjectSettings.h
@@ -21,8 +21,8 @@ class ULiveLinkRole;
 
 
 // Base class for live link subject settings
-UCLASS()
-class LIVELINKINTERFACE_API ULiveLinkSubjectSettings : public UObject
+UCLASS(MinimalAPI)
+class ULiveLinkSubjectSettings : public UObject
 {
 public:
 	GENERATED_BODY()
@@ -53,7 +53,7 @@ public:
 public:
 	//~ Begin UObject interface
 #if WITH_EDITOR
-	virtual void PostEditChangeChainProperty(struct FPropertyChangedChainEvent& PropertyChangedEvent) override;
+	LIVELINKINTERFACE_API virtual void PostEditChangeChainProperty(struct FPropertyChangedChainEvent& PropertyChangedEvent) override;
 #endif
 	//~ End UObject interface
 };
diff --git a/Engine/Source/Runtime/LiveLinkInterface/Public/LiveLinkTypes.h b/Engine/Source/Runtime/LiveLinkInterface/Public/LiveLinkTypes.h
--- a/Engine/Source/Runtime/LiveLinkInterface/Public/LiveLinkTypes.h
+++ b/Engine/Source/Runtime/LiveLinkInterface/Public/LiveLinkTypes.h
@@ -36,7 +36,7 @@
 
 
 USTRUCT(BlueprintType)
-struct LIVELINKINTERFACE_API FLiveLinkSubjectName
+struct FLiveLinkSubjectName
 {
 public:
 	GENERATED_BODY()
@@ -74,7 +74,7 @@ public:
 
 // Structure that identifies an individual subject
 USTRUCT(BlueprintType)
-struct LIVELINKINTERFACE_API FLiveLinkSubjectKey
+struct FLiveLinkSubjectKey
 {
 	GENERATED_BODY()
 
@@ -169,7 +169,7 @@ private:
 
 
 USTRUCT(BlueprintType)
-struct LIVELINKINTERFACE_API FLiveLinkTime
+struct FLiveLinkTime
 {
 	GENERATED_BODY()
 
@@ -194,7 +194,7 @@ struct LIVELINKINTERFACE_API FLiveLinkTime
 
 
 USTRUCT(BlueprintType)
-struct LIVELINKINTERFACE_API FLiveLinkMetaData
+struct FLiveLinkMetaData
 {
 	GENERATED_BODY()
 
@@ -213,7 +213,7 @@ using FLiveLinkFrameIdentifier = int32;
  * @note subclass can't contains reference to UObject
  */
 USTRUCT(BlueprintType)
-struct LIVELINKINTERFACE_API FLiveLinkBaseFrameData
+struct FLiveLinkBaseFrameData
 {
 	GENERATED_BODY();
 
@@ -246,7 +246,7 @@ struct LIVELINKINTERFACE_API FLiveLinkBaseFrameData
  * @note subclass can't contains reference to UObject
  */
 USTRUCT(BlueprintType)
-struct LIVELINKINTERFACE_API FLiveLinkBaseStaticData
+struct FLiveLinkBaseStaticData
 {
 	GENERATED_BODY()
 
@@ -276,7 +276,7 @@ struct LIVELINKINTERFACE_API FLiveLinkBaseStaticData
  * Can be used to do blueprint facilitator per role
  */
 USTRUCT(BlueprintType)
-struct LIVELINKINTERFACE_API FLiveLinkBaseBlueprintData
+struct FLiveLinkBaseBlueprintData
 {
 	GENERATED_BODY();
 
diff --git a/Engine/Source/Runtime/LiveLinkInterface/Public/LiveLinkVirtualSubject.h b/Engine/Source/Runtime/LiveLinkInterface/Public/LiveLinkVirtualSubject.h
--- a/Engine/Source/Runtime/LiveLinkInterface/Public/LiveLinkVirtualSubject.h
+++ b/Engine/Source/Runtime/LiveLinkInterface/Public/LiveLinkVirtualSubject.h
@@ -23,24 +23,24 @@ class ULiveLinkRole;
 
 
 // A Virtual subject is made up of one or more real subjects from a source
-UCLASS(Abstract)
-class LIVELINKINTERFACE_API ULiveLinkVirtualSubject : public UObject, public ILiveLinkSubject
+UCLASS(Abstract, MinimalAPI)
+class ULiveLinkVirtualSubject : public UObject, public ILiveLinkSubject
 {
 	GENERATED_BODY()
 
 	//~ Begin ILiveLinkSubject Interface
 public:
-	virtual void Initialize(FLiveLinkSubjectKey SubjectKey, TSubclassOf<ULiveLinkRole> Role, ILiveLinkClient* LiveLinkClient) override;
-	virtual void Update() override;
-	virtual bool EvaluateFrame(TSubclassOf<ULiveLinkRole> InDesiredRole, FLiveLinkSubjectFrameData& OutFrame) override;
-	virtual void ClearFrames() override;
+	LIVELINKINTERFACE_API virtual void Initialize(FLiveLinkSubjectKey SubjectKey, TSubclassOf<ULiveLinkRole> Role, ILiveLinkClient* LiveLinkClient) override;
+	LIVELINKINTERFACE_API virtual void Update() override;
+	LIVELINKINTERFACE_API virtual bool EvaluateFrame(TSubclassOf<ULiveLinkRole> InDesiredRole, FLiveLinkSubjectFrameData& OutFrame) override;
+	LIVELINKINTERFACE_API virtual void ClearFrames() override;
 	virtual FLiveLinkSubjectKey GetSubjectKey() const override { return SubjectKey; }
 	virtual TSubclassOf<ULiveLinkRole> GetRole() const override { return Role; }
-	virtual bool HasValidFrameSnapshot() const override;
+	LIVELINKINTERFACE_API virtual bool HasValidFrameSnapshot() const override;
 	virtual FLiveLinkStaticDataStruct& GetStaticData() override { return CurrentFrameSnapshot.StaticData; }
 	virtual const FLiveLinkStaticDataStruct& GetStaticData() const override { return CurrentFrameSnapshot.StaticData; }
 	virtual const TArray<ULiveLinkFrameTranslator::FWorkerSharedPtr> GetFrameTranslators() const override { return CurrentFrameTranslators; }
-	virtual TArray<FLiveLinkTime> GetFrameTimes() const override;
+	LIVELINKINTERFACE_API virtual TArray<FLiveLinkTime> GetFrameTimes() const override;
 	virtual bool IsRebroadcasted() const override { return bRebroadcastSubject; }
 	virtual bool HasStaticDataBeenRebroadcasted() const override { return bHasStaticDataBeenRebroadcast; }
 	virtual void SetStaticDataAsRebroadcasted(const bool bInSent) override { bHasStaticDataBeenRebroadcast = bInSent; }
@@ -49,10 +49,10 @@ protected:
 	//~ End ILiveLinkSubject Interface
 
 	/** Whether snapshot has valid static data */
-	bool HasValidStaticData() const;
+	LIVELINKINTERFACE_API bool HasValidStaticData() const;
 
 	/** Whether snapshot has valid frame data */
-	bool HasValidFrameData() const;
+	LIVELINKINTERFACE_API bool HasValidFrameData() const;
 
 public:
 	ILiveLinkClient* GetClient() const { return LiveLinkClient; }
@@ -67,24 +67,24 @@ public:
 	const FLiveLinkFrameDataStruct& GetFrameData() const { return CurrentFrameSnapshot.FrameData; }
 
 	/** Returns true whether this virtual subject depends on the Subject named SubjectName */
-	virtual bool DependsOnSubject(FName SubjectName) const;
+	LIVELINKINTERFACE_API virtual bool DependsOnSubject(FName SubjectName) const;
 	
 protected:
 
 	/** Updates the list of translators valid for this frame */
-	void UpdateTranslatorsForThisFrame();
+	LIVELINKINTERFACE_API void UpdateTranslatorsForThisFrame();
 
 	/** Updates our snapshot's static data */
-	void UpdateStaticDataSnapshot(FLiveLinkStaticDataStruct&& NewStaticData);
+	LIVELINKINTERFACE_API void UpdateStaticDataSnapshot(FLiveLinkStaticDataStruct&& NewStaticData);
 
 	/** Updates our snapshot's frame data */
-	void UpdateFrameDataSnapshot(FLiveLinkFrameDataStruct&& NewFrameData);
+	LIVELINKINTERFACE_API void UpdateFrameDataSnapshot(FLiveLinkFrameDataStruct&& NewFrameData);
 
 	/** Invalidates our snapshot's static data */
-	void InvalidateStaticData();
+	LIVELINKINTERFACE_API void InvalidateStaticData();
 
 	/** Invalidates our snapshot's frame data */
-	void InvalidateFrameData();
+	LIVELINKINTERFACE_API void InvalidateFrameData();
 
 protected:
 	/** The role the subject was build with. */
diff --git a/Engine/Source/Runtime/LiveLinkInterface/Public/Roles/LiveLinkAnimationBlueprintStructs.h b/Engine/Source/Runtime/LiveLinkInterface/Public/Roles/LiveLinkAnimationBlueprintStructs.h
--- a/Engine/Source/Runtime/LiveLinkInterface/Public/Roles/LiveLinkAnimationBlueprintStructs.h
+++ b/Engine/Source/Runtime/LiveLinkInterface/Public/Roles/LiveLinkAnimationBlueprintStructs.h
@@ -90,33 +90,33 @@ private:
 };
 
 USTRUCT(BlueprintType)
-struct LIVELINKINTERFACE_API FLiveLinkTransform
+struct FLiveLinkTransform
 {
 	GENERATED_USTRUCT_BODY()
 
-	FLiveLinkTransform();
+	LIVELINKINTERFACE_API FLiveLinkTransform();
 
 	virtual ~FLiveLinkTransform() = default;
 
-	void GetName(FName& Name) const;
+	LIVELINKINTERFACE_API void GetName(FName& Name) const;
 
-	void GetTransformParentSpace(FTransform& OutTransform) const;
+	LIVELINKINTERFACE_API void GetTransformParentSpace(FTransform& OutTransform) const;
 
-	void GetTransformRootSpace(FTransform& OutTransform) const;
+	LIVELINKINTERFACE_API void GetTransformRootSpace(FTransform& OutTransform) const;
 
-	bool HasParent() const;
+	LIVELINKINTERFACE_API bool HasParent() const;
 
-	void GetParent(FLiveLinkTransform& OutParentTransform) const;
+	LIVELINKINTERFACE_API void GetParent(FLiveLinkTransform& OutParentTransform) const;
 
-	int32 GetChildCount() const;
+	LIVELINKINTERFACE_API int32 GetChildCount() const;
 
-	void GetChildren(TArray<FLiveLinkTransform>& OutChildTransforms) const;
+	LIVELINKINTERFACE_API void GetChildren(TArray<FLiveLinkTransform>& OutChildTransforms) const;
 
-	void SetCachedFrame(TSharedPtr<FCachedSubjectFrame> InCachedFrame);
+	LIVELINKINTERFACE_API void SetCachedFrame(TSharedPtr<FCachedSubjectFrame> InCachedFrame);
 
-	void SetTransformIndex(const int32 InTransformIndex);
+	LIVELINKINTERFACE_API void SetTransformIndex(const int32 InTransformIndex);
 
-	int32 GetTransformIndex() const;
+	LIVELINKINTERFACE_API int32 GetTransformIndex() const;
 
 private:
 	TSharedPtr<FCachedSubjectFrame> CachedFrame;
@@ -124,7 +124,7 @@ private:
 };
 
 USTRUCT(BlueprintType)
-struct LIVELINKINTERFACE_API FSubjectFrameHandle : public FLiveLinkBaseBlueprintData
+struct FSubjectFrameHandle : public FLiveLinkBaseBlueprintData
 {
 	GENERATED_USTRUCT_BODY()
 
@@ -132,27 +132,27 @@ struct LIVELINKINTERFACE_API FSubjectFrameHandle : public FLiveLinkBaseBlueprint
 
 	virtual ~FSubjectFrameHandle() = default;
 
-	void GetCurves(TMap<FName, float>& OutCurves) const;
+	LIVELINKINTERFACE_API void GetCurves(TMap<FName, float>& OutCurves) const;
 
-	bool GetCurveValueByName(FName CurveName, float& CurveValue) const;
+	LIVELINKINTERFACE_API bool GetCurveValueByName(FName CurveName, float& CurveValue) const;
 
-	void GetSubjectMetadata(FSubjectMetadata& OutMetadata) const;
+	LIVELINKINTERFACE_API void GetSubjectMetadata(FSubjectMetadata& OutMetadata) const;
 
-	int32 GetNumberOfTransforms() const;
+	LIVELINKINTERFACE_API int32 GetNumberOfTransforms() const;
 
-	void GetTransformNames(TArray<FName>& OutTransformNames) const;
+	LIVELINKINTERFACE_API void GetTransformNames(TArray<FName>& OutTransformNames) const;
 
-	void GetRootTransform(FLiveLinkTransform& OutLiveLinkTransform) const;
+	LIVELINKINTERFACE_API void GetRootTransform(FLiveLinkTransform& OutLiveLinkTransform) const;
 
-	void GetTransformByIndex(int32 InTransformIndex, FLiveLinkTransform& OutLiveLinkTransform) const;
+	LIVELINKINTERFACE_API void GetTransformByIndex(int32 InTransformIndex, FLiveLinkTransform& OutLiveLinkTransform) const;
 
-	void GetTransformByName(FName InTransformName, FLiveLinkTransform& OutLiveLinkTransform) const;
+	LIVELINKINTERFACE_API void GetTransformByName(FName InTransformName, FLiveLinkTransform& OutLiveLinkTransform) const;
 
-	void SetCachedFrame(TSharedPtr<FCachedSubjectFrame> InCachedFrame);
+	LIVELINKINTERFACE_API void SetCachedFrame(TSharedPtr<FCachedSubjectFrame> InCachedFrame);
 
-	const FLiveLinkSkeletonStaticData* GetSourceSkeletonStaticData() const;
+	LIVELINKINTERFACE_API const FLiveLinkSkeletonStaticData* GetSourceSkeletonStaticData() const;
 
-	const FLiveLinkAnimationFrameData* GetSourceAnimationFrameData() const;
+	LIVELINKINTERFACE_API const FLiveLinkAnimationFrameData* GetSourceAnimationFrameData() const;
 
 private:
 	TSharedPtr<FCachedSubjectFrame> CachedFrame;
diff --git a/Engine/Source/Runtime/LiveLinkInterface/Public/Roles/LiveLinkAnimationRole.h b/Engine/Source/Runtime/LiveLinkInterface/Public/Roles/LiveLinkAnimationRole.h
--- a/Engine/Source/Runtime/LiveLinkInterface/Public/Roles/LiveLinkAnimationRole.h
+++ b/Engine/Source/Runtime/LiveLinkInterface/Public/Roles/LiveLinkAnimationRole.h
@@ -19,21 +19,21 @@ class UScriptStruct;
 /**
  * Role associated for Animation / Skeleton data.
  */
-UCLASS(BlueprintType, meta = (DisplayName = "Animation Role"))
-class LIVELINKINTERFACE_API ULiveLinkAnimationRole : public ULiveLinkBasicRole
+UCLASS(BlueprintType, meta = (DisplayName = "Animation Role"), MinimalAPI)
+class ULiveLinkAnimationRole : public ULiveLinkBasicRole
 {
 	GENERATED_BODY()
 
 public:
 	//~ Begin ULiveLinkRole interface
-	virtual UScriptStruct* GetStaticDataStruct() const override;
-	virtual UScriptStruct* GetFrameDataStruct() const override;
-	virtual UScriptStruct* GetBlueprintDataStruct() const override;
+	LIVELINKINTERFACE_API virtual UScriptStruct* GetStaticDataStruct() const override;
+	LIVELINKINTERFACE_API virtual UScriptStruct* GetFrameDataStruct() const override;
+	LIVELINKINTERFACE_API virtual UScriptStruct* GetBlueprintDataStruct() const override;
 
-	bool InitializeBlueprintData(const FLiveLinkSubjectFrameData& InSourceData, FLiveLinkBlueprintDataStruct& OutBlueprintData) const override;
+	LIVELINKINTERFACE_API bool InitializeBlueprintData(const FLiveLinkSubjectFrameData& InSourceData, FLiveLinkBlueprintDataStruct& OutBlueprintData) const override;
 
-	virtual FText GetDisplayName() const override;
-	virtual bool IsStaticDataValid(const FLiveLinkStaticDataStruct& InStaticData, bool& bOutShouldLogWarning) const override;
-	virtual bool IsFrameDataValid(const FLiveLinkStaticDataStruct& InStaticData, const FLiveLinkFrameDataStruct& InFrameData, bool& bOutShouldLogWarning) const override;
+	LIVELINKINTERFACE_API virtual FText GetDisplayName() const override;
+	LIVELINKINTERFACE_API virtual bool IsStaticDataValid(const FLiveLinkStaticDataStruct& InStaticData, bool& bOutShouldLogWarning) const override;
+	LIVELINKINTERFACE_API virtual bool IsFrameDataValid(const FLiveLinkStaticDataStruct& InStaticData, const FLiveLinkFrameDataStruct& InFrameData, bool& bOutShouldLogWarning) const override;
 	//~ End ULiveLinkRole interface
 };
diff --git a/Engine/Source/Runtime/LiveLinkInterface/Public/Roles/LiveLinkAnimationTypes.h b/Engine/Source/Runtime/LiveLinkInterface/Public/Roles/LiveLinkAnimationTypes.h
--- a/Engine/Source/Runtime/LiveLinkInterface/Public/Roles/LiveLinkAnimationTypes.h
+++ b/Engine/Source/Runtime/LiveLinkInterface/Public/Roles/LiveLinkAnimationTypes.h
@@ -10,7 +10,7 @@
  * Static data for Animation purposes. Contains data about bones that shouldn't change every frame.
  */
 USTRUCT(BlueprintType)
-struct LIVELINKINTERFACE_API FLiveLinkSkeletonStaticData : public FLiveLinkBaseStaticData
+struct FLiveLinkSkeletonStaticData : public FLiveLinkBaseStaticData
 {
 	GENERATED_BODY()
 
@@ -46,7 +46,7 @@ public:
  * Dynamic data for Animation purposes. 
  */
 USTRUCT(BlueprintType)
-struct LIVELINKINTERFACE_API FLiveLinkAnimationFrameData : public FLiveLinkBaseFrameData
+struct FLiveLinkAnimationFrameData : public FLiveLinkBaseFrameData
 {
 	GENERATED_BODY()
 
diff --git a/Engine/Source/Runtime/LiveLinkInterface/Public/Roles/LiveLinkBasicRole.h b/Engine/Source/Runtime/LiveLinkInterface/Public/Roles/LiveLinkBasicRole.h
--- a/Engine/Source/Runtime/LiveLinkInterface/Public/Roles/LiveLinkBasicRole.h
+++ b/Engine/Source/Runtime/LiveLinkInterface/Public/Roles/LiveLinkBasicRole.h
@@ -16,20 +16,20 @@ class UScriptStruct;
 /**
  * Role associated for no specific role data.
  */
-UCLASS(BlueprintType, meta = (DisplayName = "Basic Role"))
-class LIVELINKINTERFACE_API ULiveLinkBasicRole : public ULiveLinkRole
+UCLASS(BlueprintType, meta = (DisplayName = "Basic Role"), MinimalAPI)
+class ULiveLinkBasicRole : public ULiveLinkRole
 {
 	GENERATED_BODY()
 
 public:
 	//~ Begin ULiveLinkRole interface
-	virtual UScriptStruct* GetStaticDataStruct() const override;
-	virtual UScriptStruct* GetFrameDataStruct() const override;
-	virtual UScriptStruct* GetBlueprintDataStruct() const override;
+	LIVELINKINTERFACE_API virtual UScriptStruct* GetStaticDataStruct() const override;
+	LIVELINKINTERFACE_API virtual UScriptStruct* GetFrameDataStruct() const override;
+	LIVELINKINTERFACE_API virtual UScriptStruct* GetBlueprintDataStruct() const override;
 	
-	virtual bool InitializeBlueprintData(const FLiveLinkSubjectFrameData& InSourceData, FLiveLinkBlueprintDataStruct& OutBlueprintData) const override;
+	LIVELINKINTERFACE_API virtual bool InitializeBlueprintData(const FLiveLinkSubjectFrameData& InSourceData, FLiveLinkBlueprintDataStruct& OutBlueprintData) const override;
 
-	virtual FText GetDisplayName() const override;
-	virtual bool IsFrameDataValid(const FLiveLinkStaticDataStruct& InStaticData, const FLiveLinkFrameDataStruct& InFrameData, bool& bOutShouldLogWarning) const override;
+	LIVELINKINTERFACE_API virtual FText GetDisplayName() const override;
+	LIVELINKINTERFACE_API virtual bool IsFrameDataValid(const FLiveLinkStaticDataStruct& InStaticData, const FLiveLinkFrameDataStruct& InFrameData, bool& bOutShouldLogWarning) const override;
 	//~ End ULiveLinkRole interface
 };
diff --git a/Engine/Source/Runtime/LiveLinkInterface/Public/Roles/LiveLinkBasicTypes.h b/Engine/Source/Runtime/LiveLinkInterface/Public/Roles/LiveLinkBasicTypes.h
--- a/Engine/Source/Runtime/LiveLinkInterface/Public/Roles/LiveLinkBasicTypes.h
+++ b/Engine/Source/Runtime/LiveLinkInterface/Public/Roles/LiveLinkBasicTypes.h
@@ -9,7 +9,7 @@
  * Facility structure to handle base data in blueprint
  */
 USTRUCT(BlueprintType)
-struct LIVELINKINTERFACE_API FLiveLinkBasicBlueprintData : public FLiveLinkBaseBlueprintData
+struct FLiveLinkBasicBlueprintData : public FLiveLinkBaseBlueprintData
 {
 	GENERATED_BODY()
 
@@ -20,4 +20,4 @@ struct LIVELINKINTERFACE_API FLiveLinkBasicBlueprintData : public FLiveLinkBaseB
 	// Dynamic data that can change every frame
 	UPROPERTY(EditAnywhere, BlueprintReadWrite, Category = "Base Data")
 	FLiveLinkBaseFrameData FrameData;
-};
\ No newline at end of file
+};
diff --git a/Engine/Source/Runtime/LiveLinkInterface/Public/Roles/LiveLinkCameraRole.h b/Engine/Source/Runtime/LiveLinkInterface/Public/Roles/LiveLinkCameraRole.h
--- a/Engine/Source/Runtime/LiveLinkInterface/Public/Roles/LiveLinkCameraRole.h
+++ b/Engine/Source/Runtime/LiveLinkInterface/Public/Roles/LiveLinkCameraRole.h
@@ -16,19 +16,19 @@ class UScriptStruct;
 /**
  * Role associated for Camera data.
  */
-UCLASS(BlueprintType, meta = (DisplayName = "Camera Role"))
-class LIVELINKINTERFACE_API ULiveLinkCameraRole : public ULiveLinkTransformRole
+UCLASS(BlueprintType, meta = (DisplayName = "Camera Role"), MinimalAPI)
+class ULiveLinkCameraRole : public ULiveLinkTransformRole
 {
 	GENERATED_BODY()
 
 public:
 	//~ Begin ULiveLinkRole interface
-	virtual UScriptStruct* GetStaticDataStruct() const override;
-	virtual UScriptStruct* GetFrameDataStruct() const override;
-	virtual UScriptStruct* GetBlueprintDataStruct() const override;
+	LIVELINKINTERFACE_API virtual UScriptStruct* GetStaticDataStruct() const override;
+	LIVELINKINTERFACE_API virtual UScriptStruct* GetFrameDataStruct() const override;
+	LIVELINKINTERFACE_API virtual UScriptStruct* GetBlueprintDataStruct() const override;
 
-	virtual bool InitializeBlueprintData(const FLiveLinkSubjectFrameData& InSourceData, FLiveLinkBlueprintDataStruct& OutBlueprintData) const override;
+	LIVELINKINTERFACE_API virtual bool InitializeBlueprintData(const FLiveLinkSubjectFrameData& InSourceData, FLiveLinkBlueprintDataStruct& OutBlueprintData) const override;
 
-	virtual FText GetDisplayName() const override;
+	LIVELINKINTERFACE_API virtual FText GetDisplayName() const override;
 	//~ End ULiveLinkRole interface
 };
diff --git a/Engine/Source/Runtime/LiveLinkInterface/Public/Roles/LiveLinkCameraTypes.h b/Engine/Source/Runtime/LiveLinkInterface/Public/Roles/LiveLinkCameraTypes.h
--- a/Engine/Source/Runtime/LiveLinkInterface/Public/Roles/LiveLinkCameraTypes.h
+++ b/Engine/Source/Runtime/LiveLinkInterface/Public/Roles/LiveLinkCameraTypes.h
@@ -17,7 +17,7 @@ enum class ELiveLinkCameraProjectionMode : uint8
  * Static data for Camera data. 
  */
 USTRUCT(BlueprintType)
-struct LIVELINKINTERFACE_API FLiveLinkCameraStaticData : public FLiveLinkTransformStaticData
+struct FLiveLinkCameraStaticData : public FLiveLinkTransformStaticData
 {
 	GENERATED_BODY()
 
@@ -64,7 +64,7 @@ public:
  * Dynamic data for camera 
  */
 USTRUCT(BlueprintType)
-struct LIVELINKINTERFACE_API FLiveLinkCameraFrameData : public FLiveLinkTransformFrameData
+struct FLiveLinkCameraFrameData : public FLiveLinkTransformFrameData
 {
 	GENERATED_BODY()
 
@@ -97,7 +97,7 @@ struct LIVELINKINTERFACE_API FLiveLinkCameraFrameData : public FLiveLinkTransfor
  * Facility structure to handle camera data in blueprint
  */
 USTRUCT(BlueprintType)
-struct LIVELINKINTERFACE_API FLiveLinkCameraBlueprintData : public FLiveLinkBaseBlueprintData
+struct FLiveLinkCameraBlueprintData : public FLiveLinkBaseBlueprintData
 {
 	GENERATED_BODY()
 	
@@ -108,4 +108,4 @@ struct LIVELINKINTERFACE_API FLiveLinkCameraBlueprintData : public FLiveLinkBase
 	// Dynamic data that can change every frame
 	UPROPERTY(EditAnywhere, BlueprintReadWrite, Category = "LiveLink")
 	FLiveLinkCameraFrameData FrameData;
-};
\ No newline at end of file
+};
diff --git a/Engine/Source/Runtime/LiveLinkInterface/Public/Roles/LiveLinkLightRole.h b/Engine/Source/Runtime/LiveLinkInterface/Public/Roles/LiveLinkLightRole.h
--- a/Engine/Source/Runtime/LiveLinkInterface/Public/Roles/LiveLinkLightRole.h
+++ b/Engine/Source/Runtime/LiveLinkInterface/Public/Roles/LiveLinkLightRole.h
@@ -16,19 +16,19 @@ class UScriptStruct;
 /**
  * Role associated for Light data.
  */
-UCLASS(BlueprintType, meta = (DisplayName = "Light Role"))
-class LIVELINKINTERFACE_API ULiveLinkLightRole : public ULiveLinkTransformRole
+UCLASS(BlueprintType, meta = (DisplayName = "Light Role"), MinimalAPI)
+class ULiveLinkLightRole : public ULiveLinkTransformRole
 {
 	GENERATED_BODY()
 
 public:
 	//~ Begin ULiveLinkRole interface
-	virtual UScriptStruct* GetStaticDataStruct() const override;
-	virtual UScriptStruct* GetFrameDataStruct() const override;
-	virtual UScriptStruct* GetBlueprintDataStruct() const override;
+	LIVELINKINTERFACE_API virtual UScriptStruct* GetStaticDataStruct() const override;
+	LIVELINKINTERFACE_API virtual UScriptStruct* GetFrameDataStruct() const override;
+	LIVELINKINTERFACE_API virtual UScriptStruct* GetBlueprintDataStruct() const override;
 
-	bool InitializeBlueprintData(const FLiveLinkSubjectFrameData& InSourceData, FLiveLinkBlueprintDataStruct& OutBlueprintData) const override;
+	LIVELINKINTERFACE_API bool InitializeBlueprintData(const FLiveLinkSubjectFrameData& InSourceData, FLiveLinkBlueprintDataStruct& OutBlueprintData) const override;
 
-	virtual FText GetDisplayName() const override;
+	LIVELINKINTERFACE_API virtual FText GetDisplayName() const override;
 	//~ End ULiveLinkRole interface
 };
diff --git a/Engine/Source/Runtime/LiveLinkInterface/Public/Roles/LiveLinkLightTypes.h b/Engine/Source/Runtime/LiveLinkInterface/Public/Roles/LiveLinkLightTypes.h
--- a/Engine/Source/Runtime/LiveLinkInterface/Public/Roles/LiveLinkLightTypes.h
+++ b/Engine/Source/Runtime/LiveLinkInterface/Public/Roles/LiveLinkLightTypes.h
@@ -9,7 +9,7 @@
  * Static data for Light data. 
  */
 USTRUCT(BlueprintType)
-struct LIVELINKINTERFACE_API FLiveLinkLightStaticData : public FLiveLinkTransformStaticData
+struct FLiveLinkLightStaticData : public FLiveLinkTransformStaticData
 {
 	GENERATED_BODY()
 
@@ -56,7 +56,7 @@ public:
  * Dynamic data for light. 
  */
 USTRUCT(BlueprintType)
-struct LIVELINKINTERFACE_API FLiveLinkLightFrameData : public FLiveLinkTransformFrameData
+struct FLiveLinkLightFrameData : public FLiveLinkTransformFrameData
 {
 	GENERATED_BODY()
 
@@ -101,7 +101,7 @@ struct LIVELINKINTERFACE_API FLiveLinkLightFrameData : public FLiveLinkTransform
  * Facility structure to handle light data in blueprint
  */
 USTRUCT(BlueprintType)
-struct LIVELINKINTERFACE_API FLiveLinkLightBlueprintData : public FLiveLinkBaseBlueprintData
+struct FLiveLinkLightBlueprintData : public FLiveLinkBaseBlueprintData
 {
 	GENERATED_BODY()
 	
diff --git a/Engine/Source/Runtime/LiveLinkInterface/Public/Roles/LiveLinkTransformRole.h b/Engine/Source/Runtime/LiveLinkInterface/Public/Roles/LiveLinkTransformRole.h
--- a/Engine/Source/Runtime/LiveLinkInterface/Public/Roles/LiveLinkTransformRole.h
+++ b/Engine/Source/Runtime/LiveLinkInterface/Public/Roles/LiveLinkTransformRole.h
@@ -16,20 +16,20 @@ class UScriptStruct;
 /**
  * Role associated for Camera data.
  */
-UCLASS(BlueprintType, meta = (DisplayName = "Transform Role"))
-class LIVELINKINTERFACE_API ULiveLinkTransformRole : public ULiveLinkBasicRole
+UCLASS(BlueprintType, meta = (DisplayName = "Transform Role"), MinimalAPI)
+class ULiveLinkTransformRole : public ULiveLinkBasicRole
 {
 	GENERATED_BODY()
 
 public:
 	//~ Begin ULiveLinkRole interface
-	virtual UScriptStruct* GetStaticDataStruct() const override;
-	virtual UScriptStruct* GetFrameDataStruct() const override;
-	virtual UScriptStruct* GetBlueprintDataStruct() const override;
+	LIVELINKINTERFACE_API virtual UScriptStruct* GetStaticDataStruct() const override;
+	LIVELINKINTERFACE_API virtual UScriptStruct* GetFrameDataStruct() const override;
+	LIVELINKINTERFACE_API virtual UScriptStruct* GetBlueprintDataStruct() const override;
 	
-	virtual bool InitializeBlueprintData(const FLiveLinkSubjectFrameData& InSourceData, FLiveLinkBlueprintDataStruct& OutBlueprintData) const override;
+	LIVELINKINTERFACE_API virtual bool InitializeBlueprintData(const FLiveLinkSubjectFrameData& InSourceData, FLiveLinkBlueprintDataStruct& OutBlueprintData) const override;
 
-	virtual FText GetDisplayName() const override;
+	LIVELINKINTERFACE_API virtual FText GetDisplayName() const override;
 	//~ End ULiveLinkRole interface
 };
 
diff --git a/Engine/Source/Runtime/LiveLinkInterface/Public/Roles/LiveLinkTransformTypes.h b/Engine/Source/Runtime/LiveLinkInterface/Public/Roles/LiveLinkTransformTypes.h
--- a/Engine/Source/Runtime/LiveLinkInterface/Public/Roles/LiveLinkTransformTypes.h
+++ b/Engine/Source/Runtime/LiveLinkInterface/Public/Roles/LiveLinkTransformTypes.h
@@ -9,7 +9,7 @@
  * Static data for Transform data. 
  */
 USTRUCT(BlueprintType)
-struct LIVELINKINTERFACE_API FLiveLinkTransformStaticData : public FLiveLinkBaseStaticData
+struct FLiveLinkTransformStaticData : public FLiveLinkBaseStaticData
 {
 	GENERATED_BODY()
 	
@@ -30,7 +30,7 @@ struct LIVELINKINTERFACE_API FLiveLinkTransformStaticData : public FLiveLinkBase
  * Dynamic data for Transform 
  */
 USTRUCT(BlueprintType)
-struct LIVELINKINTERFACE_API FLiveLinkTransformFrameData : public FLiveLinkBaseFrameData
+struct FLiveLinkTransformFrameData : public FLiveLinkBaseFrameData
 {
 	GENERATED_BODY()
 
@@ -43,7 +43,7 @@ struct LIVELINKINTERFACE_API FLiveLinkTransformFrameData : public FLiveLinkBaseF
  * Facility structure to handle transform data in blueprint
  */
 USTRUCT(BlueprintType)
-struct LIVELINKINTERFACE_API FLiveLinkTransformBlueprintData : public FLiveLinkBaseBlueprintData
+struct FLiveLinkTransformBlueprintData : public FLiveLinkBaseBlueprintData
 {
 	GENERATED_BODY()
 
diff --git a/Engine/Source/Runtime/MRMesh/Public/MRMeshComponent.h b/Engine/Source/Runtime/MRMesh/Public/MRMeshComponent.h
--- a/Engine/Source/Runtime/MRMesh/Public/MRMeshComponent.h
+++ b/Engine/Source/Runtime/MRMesh/Public/MRMeshComponent.h
@@ -107,16 +107,16 @@ public:
 
 DECLARE_MULTICAST_DELEGATE_TwoParams(FOnMRMeshBrickDataUpdatedDelegate, const UMRMeshComponent*, const IMRMesh::FSendBrickDataArgs&);
 
-UCLASS(hideCategories=(Physics), meta = (BlueprintSpawnableComponent), ClassGroup = Rendering)
-class MRMESH_API UMRMeshComponent : public UPrimitiveComponent, public IMRMesh
+UCLASS(hideCategories=(Physics), meta = (BlueprintSpawnableComponent), ClassGroup = Rendering, MinimalAPI)
+class UMRMeshComponent : public UPrimitiveComponent, public IMRMesh
 {
 public:
 	friend class FMRMeshProxy;
 
 	GENERATED_UCLASS_BODY()
 
-	virtual void BeginPlay() override;
-	void EndPlay(const EEndPlayReason::Type EndPlayReason) override;
+	MRMESH_API virtual void BeginPlay() override;
+	MRMESH_API void EndPlay(const EEndPlayReason::Type EndPlayReason) override;
 
 	UFUNCTION(BlueprintPure, Category = "Mesh Reconstruction")
 	bool IsConnected() const override { return bConnected; }
@@ -128,67 +128,67 @@ public:
 	* Force navmesh generation to run using the current collision data.  This will run even if the collision data has not been udpated! Unless you are changing navmesh settings or similar RequestNavMeshUpdate is reccomended.
 	*/	
 	UFUNCTION(BlueprintCallable, Category = "Mesh Reconstruction")
-	void ForceNavMeshUpdate();
+	MRMESH_API void ForceNavMeshUpdate();
 
 	/**
 	* Generate nav mesh if collision data has changed since the last nav mesh generation.  
 	*/
 	UFUNCTION(BlueprintCallable, Category = "Mesh Reconstruction")
-	void RequestNavMeshUpdate();
+	MRMESH_API void RequestNavMeshUpdate();
 
 	UFUNCTION(BlueprintCallable, Category = "Mesh Reconstruction")
-	void Clear() override;
+	MRMESH_API void Clear() override;
 
 	// UPrimitiveComponent.. public BP function needs to stay public to avoid nativization errors. (RR)
-	virtual void SetMaterial(int32 ElementIndex, class UMaterialInterface* InMaterial) override;
-	virtual class UMaterialInterface* GetMaterial(int32 ElementIndex) const override;
+	MRMESH_API virtual void SetMaterial(int32 ElementIndex, class UMaterialInterface* InMaterial) override;
+	MRMESH_API virtual class UMaterialInterface* GetMaterial(int32 ElementIndex) const override;
 
 	// Set the wireframe material.
 	UFUNCTION(BlueprintCallable, Category = "Mesh Reconstruction")
-	virtual void SetWireframeMaterial(class UMaterialInterface* InMaterial);
+	MRMESH_API virtual void SetWireframeMaterial(class UMaterialInterface* InMaterial);
 
 	/** Updates from HoloLens or iOS */
-	void UpdateMesh(const FVector& InLocation, const FQuat& InRotation, const FVector& Scale, TArray<FVector>& Vertices, TArray<MRMESH_INDEX_TYPE>& Indices, TArray<FVector2D> UVData = {}, TArray<FPackedNormal> TangentXZData = {}, TArray<FColor> ColorData = {});
-	void UpdateMesh(const FVector& InLocation, const FQuat& InRotation, const FVector& Scale, TArray<FVector3f>& Vertices, TArray<MRMESH_INDEX_TYPE>& Indices, TArray<FVector2D> UVData = {}, TArray<FPackedNormal> TangentXZData = {}, TArray<FColor> ColorData = {});
+	MRMESH_API void UpdateMesh(const FVector& InLocation, const FQuat& InRotation, const FVector& Scale, TArray<FVector>& Vertices, TArray<MRMESH_INDEX_TYPE>& Indices, TArray<FVector2D> UVData = {}, TArray<FPackedNormal> TangentXZData = {}, TArray<FColor> ColorData = {});
+	MRMESH_API void UpdateMesh(const FVector& InLocation, const FQuat& InRotation, const FVector& Scale, TArray<FVector3f>& Vertices, TArray<MRMESH_INDEX_TYPE>& Indices, TArray<FVector2D> UVData = {}, TArray<FPackedNormal> TangentXZData = {}, TArray<FColor> ColorData = {});
 
 	UFUNCTION(BlueprintCallable, Category = "Mesh Reconstruction")
-	void SetEnableMeshOcclusion(bool bEnable);
+	MRMESH_API void SetEnableMeshOcclusion(bool bEnable);
 	
 	UFUNCTION(BlueprintPure, Category = "Mesh Reconstruction")
 	bool GetEnableMeshOcclusion() const { return bEnableOcclusion; }
 	
 	UFUNCTION(BlueprintCallable, Category = "Mesh Reconstruction")
-	void SetUseWireframe(bool bUseWireframe);
+	MRMESH_API void SetUseWireframe(bool bUseWireframe);
 	
 	UFUNCTION(BlueprintPure, Category = "Mesh Reconstruction")
 	bool GetUseWireframe() const { return bUseWireframe; }
 	
 	UFUNCTION(BlueprintCallable, Category = "Mesh Reconstruction")
-	void SetWireframeColor(const FLinearColor& InColor);
+	MRMESH_API void SetWireframeColor(const FLinearColor& InColor);
 	
 	UFUNCTION(BlueprintPure, Category = "Mesh Reconstruction")
 	const FLinearColor& GetWireframeColor() const { return WireframeColor; }
 	
-	UMaterialInterface* GetMaterialToUse() const;
+	MRMESH_API UMaterialInterface* GetMaterialToUse() const;
 	
 	FOnMRMeshBrickDataUpdatedDelegate& OnBrickDataUpdated() { return OnBrickDataUpdatedDelegate; }
 
 protected:
-	virtual void OnActorEnableCollisionChanged() override;
-	virtual void UpdatePhysicsToRBChannels() override;
+	MRMESH_API virtual void OnActorEnableCollisionChanged() override;
+	MRMESH_API virtual void UpdatePhysicsToRBChannels() override;
 public:
-	virtual void SetCollisionObjectType(ECollisionChannel Channel) override;
-	virtual void SetCollisionResponseToChannel(ECollisionChannel Channel, ECollisionResponse NewResponse) override;
-	virtual void SetCollisionResponseToAllChannels(ECollisionResponse NewResponse) override;
-	virtual void SetCollisionResponseToChannels(const FCollisionResponseContainer& NewResponses) override;
-	virtual void SetCollisionEnabled(ECollisionEnabled::Type NewType) override;
-	virtual void SetCollisionProfileName(FName InCollisionProfileName, bool bUpdateOverlaps=true) override;
+	MRMESH_API virtual void SetCollisionObjectType(ECollisionChannel Channel) override;
+	MRMESH_API virtual void SetCollisionResponseToChannel(ECollisionChannel Channel, ECollisionResponse NewResponse) override;
+	MRMESH_API virtual void SetCollisionResponseToAllChannels(ECollisionResponse NewResponse) override;
+	MRMESH_API virtual void SetCollisionResponseToChannels(const FCollisionResponseContainer& NewResponses) override;
+	MRMESH_API virtual void SetCollisionEnabled(ECollisionEnabled::Type NewType) override;
+	MRMESH_API virtual void SetCollisionProfileName(FName InCollisionProfileName, bool bUpdateOverlaps=true) override;
 
-	virtual void SetWalkableSlopeOverride(const FWalkableSlopeOverride& NewOverride) override;
+	MRMESH_API virtual void SetWalkableSlopeOverride(const FWalkableSlopeOverride& NewOverride) override;
 
 	void SetNeverCreateCollisionMesh(bool bNeverCreate) { bNeverCreateCollisionMesh = bNeverCreate; }
 	void SetEnableNavMesh(bool bEnable) { bUpdateNavMeshOnMeshUpdate = bEnable;  }
-	void SuggestNavMeshUpdate();
+	MRMESH_API void SuggestNavMeshUpdate();
 
 	/** Trackers feeding mesh data to this component may want to know when we clear our mesh data */
 	DECLARE_EVENT(UMRMeshComponent, FOnClear);
@@ -196,29 +196,29 @@ public:
 
 private:
 	//~ UPrimitiveComponent
-	virtual FPrimitiveSceneProxy* CreateSceneProxy() override;
-	virtual void GetUsedMaterials(TArray<UMaterialInterface*>& OutMaterials, bool bGetDebugMaterials = false) const override;
-	virtual FBoxSphereBounds CalcBounds(const FTransform& LocalToWorld) const override;
-	virtual bool DoCustomNavigableGeometryExport(FNavigableGeometryExport& GeomExport) const override;
+	MRMESH_API virtual FPrimitiveSceneProxy* CreateSceneProxy() override;
+	MRMESH_API virtual void GetUsedMaterials(TArray<UMaterialInterface*>& OutMaterials, bool bGetDebugMaterials = false) const override;
+	MRMESH_API virtual FBoxSphereBounds CalcBounds(const FTransform& LocalToWorld) const override;
+	MRMESH_API virtual bool DoCustomNavigableGeometryExport(FNavigableGeometryExport& GeomExport) const override;
 
 	//~ UPrimitiveComponent
 	//~ UActorComponent
-	virtual bool ShouldCreatePhysicsState() const override;
-	virtual void SendRenderDynamicData_Concurrent() override;
+	MRMESH_API virtual bool ShouldCreatePhysicsState() const override;
+	MRMESH_API virtual void SendRenderDynamicData_Concurrent() override;
 	//~ UActorComponent
 
 	//~ IMRMesh
-	virtual void SendBrickData(FSendBrickDataArgs Args) override;
-	virtual void ClearAllBrickData() override;
+	MRMESH_API virtual void SendBrickData(FSendBrickDataArgs Args) override;
+	MRMESH_API virtual void ClearAllBrickData() override;
 	//~ IMRMesh
 
 private:
-	void SendBrickData_Internal(IMRMesh::FSendBrickDataArgs Args);
+	MRMESH_API void SendBrickData_Internal(IMRMesh::FSendBrickDataArgs Args);
 
-	void RemoveBodyInstance(int32 BodyIndex);
-	void OnUpdateTransform(EUpdateTransformFlags UpdateTransformFlags, ETeleportType Teleport) override;
+	MRMESH_API void RemoveBodyInstance(int32 BodyIndex);
+	MRMESH_API void OnUpdateTransform(EUpdateTransformFlags UpdateTransformFlags, ETeleportType Teleport) override;
 
-	void ClearAllBrickData_Internal();
+	MRMESH_API void ClearAllBrickData_Internal();
 
 	UPROPERTY(EditAnywhere, Category = Appearance)
 	TObjectPtr<UMaterialInterface> Material;
diff --git a/Engine/Source/Runtime/MRMesh/Public/MeshReconstructorBase.h b/Engine/Source/Runtime/MRMesh/Public/MeshReconstructorBase.h
--- a/Engine/Source/Runtime/MRMesh/Public/MeshReconstructorBase.h
+++ b/Engine/Source/Runtime/MRMesh/Public/MeshReconstructorBase.h
@@ -21,31 +21,31 @@ struct FMRMeshConfiguration
 	bool bSendVertexColors = false;
 };
 
-UCLASS(meta=(Experimental))
-class MRMESH_API UMeshReconstructorBase : public UObject
+UCLASS(meta=(Experimental), MinimalAPI)
+class UMeshReconstructorBase : public UObject
 {
 	GENERATED_BODY()
 
 public:
 	UFUNCTION(BlueprintCallable, Category = "Mesh Reconstruction")
-	virtual void StartReconstruction();
+	MRMESH_API virtual void StartReconstruction();
 
 	UFUNCTION(BlueprintCallable, Category = "Mesh Reconstruction")
-	virtual void StopReconstruction();
+	MRMESH_API virtual void StopReconstruction();
 
 	UFUNCTION(BlueprintCallable, Category = "Mesh Reconstruction")
-	virtual void PauseReconstruction();
+	MRMESH_API virtual void PauseReconstruction();
 
 	UFUNCTION(BlueprintCallable, Category = "Mesh Reconstruction")
-	virtual bool IsReconstructionStarted() const;
+	MRMESH_API virtual bool IsReconstructionStarted() const;
 
 	UFUNCTION(BlueprintCallable, Category = "Mesh Reconstruction")
-	virtual bool IsReconstructionPaused() const;
+	MRMESH_API virtual bool IsReconstructionPaused() const;
 
 	UFUNCTION()
-	virtual void ConnectMRMesh(UMRMeshComponent* Mesh);
+	MRMESH_API virtual void ConnectMRMesh(UMRMeshComponent* Mesh);
 
 	UFUNCTION()
-	virtual void DisconnectMRMesh();
+	MRMESH_API virtual void DisconnectMRMesh();
 
-};
\ No newline at end of file
+};
diff --git a/Engine/Source/Runtime/MRMesh/Public/MockDataMeshTrackerComponent.h b/Engine/Source/Runtime/MRMesh/Public/MockDataMeshTrackerComponent.h
--- a/Engine/Source/Runtime/MRMesh/Public/MockDataMeshTrackerComponent.h
+++ b/Engine/Source/Runtime/MRMesh/Public/MockDataMeshTrackerComponent.h
@@ -28,29 +28,29 @@ enum class EMeshTrackerVertexColorMode : uint8
 	on a separate thread.  Once a mesh data request has been processed the calling system will be notified via an
 	FOnMeshTrackerUpdated broadcast.
 */
-UCLASS(ClassGroup = Rendering, BlueprintType, Blueprintable, EditInlineNew, meta = (BlueprintSpawnableComponent))
-class MRMESH_API UMockDataMeshTrackerComponent
+UCLASS(ClassGroup = Rendering, BlueprintType, Blueprintable, EditInlineNew, meta = (BlueprintSpawnableComponent), MinimalAPI)
+class UMockDataMeshTrackerComponent
 	: public USceneComponent
 {
 	GENERATED_UCLASS_BODY()
 
 public:
 	/** Destroys the FMeshTrackerImpl instance.*/
-	~UMockDataMeshTrackerComponent();
+	MRMESH_API ~UMockDataMeshTrackerComponent();
 
 	/**
 		Sets the procedural mesh component that will store and display the environmental mesh results.
 		@param InMRMeshPtr The procedural mesh component to store the query result in.
 	*/
 	UFUNCTION(BlueprintCallable, Category = "Mesh Reconstruction")
-	void ConnectMRMesh(class UMRMeshComponent* InMRMeshPtr);
+	MRMESH_API void ConnectMRMesh(class UMRMeshComponent* InMRMeshPtr);
 
 	/**
 		Unlinks the current procedural mesh component from the mesh tracking system.
 		@param InMRMeshPtr The procedural mesh component to unlink from the mesh tracking system.
 	*/
 	UFUNCTION(BlueprintCallable, Category = "Mesh Reconstruction")
-	void DisconnectMRMesh(class UMRMeshComponent* InMRMeshPtr);
+	MRMESH_API void DisconnectMRMesh(class UMRMeshComponent* InMRMeshPtr);
 
 	/**
 	Delegate used by OnMeshUpdated().
@@ -118,16 +118,16 @@ public:
 	*/
 
 	/** Polls for and handles the results of the environmental mesh queries. */
-	virtual void TickComponent(float DeltaTime, enum ELevelTick TickType, FActorComponentTickFunction *ThisTickFunction) override;
+	MRMESH_API virtual void TickComponent(float DeltaTime, enum ELevelTick TickType, FActorComponentTickFunction *ThisTickFunction) override;
 
 	/** Unlinks the current procedural mesh component from the mesh tracking system. */
-	virtual void BeginDestroy() override;
+	MRMESH_API virtual void BeginDestroy() override;
 
 	/** Destroys the interface object to the mesh tracking api*/
-	virtual void FinishDestroy() override;
+	MRMESH_API virtual void FinishDestroy() override;
 
 #if WITH_EDITOR
-	virtual void PostEditChangeProperty(FPropertyChangedEvent& e) override;
+	MRMESH_API virtual void PostEditChangeProperty(FPropertyChangedEvent& e) override;
 #endif
 
 private:
diff --git a/Engine/Source/Runtime/MaterialShaderQualitySettings/Classes/MaterialShaderQualitySettings.h b/Engine/Source/Runtime/MaterialShaderQualitySettings/Classes/MaterialShaderQualitySettings.h
--- a/Engine/Source/Runtime/MaterialShaderQualitySettings/Classes/MaterialShaderQualitySettings.h
+++ b/Engine/Source/Runtime/MaterialShaderQualitySettings/Classes/MaterialShaderQualitySettings.h
@@ -15,27 +15,27 @@ class UShaderPlatformQualitySettings;
 enum EShaderPlatform : uint16;
 
 //UCLASS(config = Engine, defaultconfig)
-UCLASS()
-class MATERIALSHADERQUALITYSETTINGS_API UMaterialShaderQualitySettings : public UObject
+UCLASS(MinimalAPI)
+class UMaterialShaderQualitySettings : public UObject
 {
 public:
 	GENERATED_UCLASS_BODY()
 
-	UShaderPlatformQualitySettings* GetShaderPlatformQualitySettings(FName PlatformName);
+	MATERIALSHADERQUALITYSETTINGS_API UShaderPlatformQualitySettings* GetShaderPlatformQualitySettings(FName PlatformName);
 
-	const UShaderPlatformQualitySettings* GetShaderPlatformQualitySettings(EShaderPlatform ShaderPlatform);
+	MATERIALSHADERQUALITYSETTINGS_API const UShaderPlatformQualitySettings* GetShaderPlatformQualitySettings(EShaderPlatform ShaderPlatform);
 
-	bool HasPlatformQualitySettings(EShaderPlatform ShaderPlatform, EMaterialQualityLevel::Type QualityLevel);
+	MATERIALSHADERQUALITYSETTINGS_API bool HasPlatformQualitySettings(EShaderPlatform ShaderPlatform, EMaterialQualityLevel::Type QualityLevel);
 
 #if WITH_EDITOR
 	// Override GetShaderPlatformQualitySettings() return value with the specified platform's settings.
 	// An empty PlatformName or otherwise non existent platform will cause GetShaderPlatformQualitySettings() 
 	// to revert to its default behaviour.
-	void SetPreviewPlatform(FName PlatformName);
-	const FName& GetPreviewPlatform();
+	MATERIALSHADERQUALITYSETTINGS_API void SetPreviewPlatform(FName PlatformName);
+	MATERIALSHADERQUALITYSETTINGS_API const FName& GetPreviewPlatform();
 #endif
 
-	static UMaterialShaderQualitySettings* Get();
+	static MATERIALSHADERQUALITYSETTINGS_API UMaterialShaderQualitySettings* Get();
 
 private:
 	UShaderPlatformQualitySettings* GetOrCreatePlatformSettings(FName ShaderPlatformName);
diff --git a/Engine/Source/Runtime/MaterialShaderQualitySettings/Classes/ShaderPlatformQualitySettings.h b/Engine/Source/Runtime/MaterialShaderQualitySettings/Classes/ShaderPlatformQualitySettings.h
--- a/Engine/Source/Runtime/MaterialShaderQualitySettings/Classes/ShaderPlatformQualitySettings.h
+++ b/Engine/Source/Runtime/MaterialShaderQualitySettings/Classes/ShaderPlatformQualitySettings.h
@@ -77,8 +77,8 @@ public:
 };
 
 
-UCLASS(config = Engine, defaultconfig, perObjectConfig)
-class MATERIALSHADERQUALITYSETTINGS_API UShaderPlatformQualitySettings : public UObject
+UCLASS(config = Engine, defaultconfig, perObjectConfig, MinimalAPI)
+class UShaderPlatformQualitySettings : public UObject
 {
 public:
 	GENERATED_UCLASS_BODY()
@@ -92,9 +92,9 @@ public:
 		return QualityOverrides[(int32)QualityLevel];
 	}
 
-	const FMaterialQualityOverrides& GetQualityOverrides(EMaterialQualityLevel::Type QualityLevel) const;
-	void BuildHash(EMaterialQualityLevel::Type QualityLevel, class FSHAHash& OutHash) const;
-	void AppendToHashState(EMaterialQualityLevel::Type QualityLevel, class FSHA1& HashState) const;
+	MATERIALSHADERQUALITYSETTINGS_API const FMaterialQualityOverrides& GetQualityOverrides(EMaterialQualityLevel::Type QualityLevel) const;
+	MATERIALSHADERQUALITYSETTINGS_API void BuildHash(EMaterialQualityLevel::Type QualityLevel, class FSHAHash& OutHash) const;
+	MATERIALSHADERQUALITYSETTINGS_API void AppendToHashState(EMaterialQualityLevel::Type QualityLevel, class FSHA1& HashState) const;
 	
 	virtual const TCHAR* GetConfigOverridePlatform() const override
 	{
diff --git a/Engine/Source/Runtime/Media/Public/IMediaView.h b/Engine/Source/Runtime/Media/Public/IMediaView.h
--- a/Engine/Source/Runtime/Media/Public/IMediaView.h
+++ b/Engine/Source/Runtime/Media/Public/IMediaView.h
@@ -9,7 +9,7 @@
 /**
  * Tile coordinate of a tiled media source/texture.
  */
-struct MEDIA_API FMediaTileCoordinate
+struct FMediaTileCoordinate
 {
 	uint16 X, Y;
 
@@ -253,4 +253,4 @@ public:
 
 	/** Virtual destructor. */
 	virtual ~IMediaView() { }
-};
\ No newline at end of file
+};
diff --git a/Engine/Source/Runtime/MediaAssets/Public/BaseMediaSource.h b/Engine/Source/Runtime/MediaAssets/Public/BaseMediaSource.h
--- a/Engine/Source/Runtime/MediaAssets/Public/BaseMediaSource.h
+++ b/Engine/Source/Runtime/MediaAssets/Public/BaseMediaSource.h
@@ -24,8 +24,8 @@ struct FGuid;
 /**
  * Base class for concrete media sources.
  */
-UCLASS(Abstract, BlueprintType, hidecategories=(Object))
-class MEDIAASSETS_API UBaseMediaSource
+UCLASS(Abstract, BlueprintType, hidecategories=(Object), MinimalAPI)
+class UBaseMediaSource
 	: public UMediaSource
 {
 	GENERATED_BODY()
@@ -49,17 +49,17 @@ public:
 	//~ UObject interface
 	PRAGMA_DISABLE_DEPRECATION_WARNINGS // Suppress compiler warning on override of deprecated function
 	UE_DEPRECATED(5.0, "Use version that takes FObjectPreSaveContext instead.")
-	virtual void PreSave(const class ITargetPlatform* TargetPlatform);
-	PRAGMA_ENABLE_DEPRECATION_WARNINGS
+	MEDIAASSETS_API virtual void PreSave(const class ITargetPlatform* TargetPlatform);
+	MEDIAASSETS_API PRAGMA_ENABLE_DEPRECATION_WARNINGS
 	virtual void PreSave(FObjectPreSaveContext ObjectSaveContext);
-	virtual void GetAssetRegistryTags(TArray<FAssetRegistryTag>& OutTags) const override;
-	virtual void Serialize(FArchive& Ar) override;
+	MEDIAASSETS_API virtual void GetAssetRegistryTags(TArray<FAssetRegistryTag>& OutTags) const override;
+	MEDIAASSETS_API virtual void Serialize(FArchive& Ar) override;
 
 public:
 
 	//~ IMediaOptions interface
 
-	virtual FName GetDesiredPlayerName() const override;
+	MEDIAASSETS_API virtual FName GetDesiredPlayerName() const override;
 
 private:
 
diff --git a/Engine/Source/Runtime/MediaAssets/Public/FileMediaSource.h b/Engine/Source/Runtime/MediaAssets/Public/FileMediaSource.h
--- a/Engine/Source/Runtime/MediaAssets/Public/FileMediaSource.h
+++ b/Engine/Source/Runtime/MediaAssets/Public/FileMediaSource.h
@@ -17,8 +17,8 @@ struct FFrame;
 struct FPropertyChangedEvent;
 
 
-UCLASS(BlueprintType)
-class MEDIAASSETS_API UFileMediaSource
+UCLASS(BlueprintType, MinimalAPI)
+class UFileMediaSource
 	: public UBaseMediaSource
 {
 	GENERATED_BODY()
@@ -57,7 +57,7 @@ public:
 	 * @return The full file path.
 	 * @return GetFilePath
 	 */
-	FString GetFullPath() const;
+	MEDIAASSETS_API FString GetFullPath() const;
 
 	/**
 	 * Set the path to the media file that this source represents.
@@ -69,33 +69,33 @@ public:
 	 * @see FilePath, GetFilePath
 	 */
 	UFUNCTION(BlueprintCallable, Category="Media|FileMediaSource")
-	void SetFilePath(const FString& Path);
+	MEDIAASSETS_API void SetFilePath(const FString& Path);
 
 #if WITH_EDITOR
-	virtual void PostEditChangeProperty(FPropertyChangedEvent& PropertyChangedEvent) override;
+	MEDIAASSETS_API virtual void PostEditChangeProperty(FPropertyChangedEvent& PropertyChangedEvent) override;
 #endif
 
 public:
 
 	//~ IMediaOptions interface
 
-	virtual bool GetMediaOption(const FName& Key, bool DefaultValue) const override;
-	virtual bool HasMediaOption(const FName& Key) const override;
+	MEDIAASSETS_API virtual bool GetMediaOption(const FName& Key, bool DefaultValue) const override;
+	MEDIAASSETS_API virtual bool HasMediaOption(const FName& Key) const override;
 
 public:
 
 	//~ UMediaSource interface
 
-	virtual FString GetUrl() const override;
-	virtual bool Validate() const override;
+	MEDIAASSETS_API virtual FString GetUrl() const override;
+	MEDIAASSETS_API virtual bool Validate() const override;
 
 public:
 
 	/** Name of the PrecacheFile media option. */
-	static FName PrecacheFileOption;
+	static MEDIAASSETS_API FName PrecacheFileOption;
 
 private:
-	void ResolveFullPath() const;
-	void ClearResolvedFullPath() const;
+	MEDIAASSETS_API void ResolveFullPath() const;
+	MEDIAASSETS_API void ClearResolvedFullPath() const;
 	mutable FString ResolvedFullPath; // this is a cached variable updated in ResolveFullPath hence mutable
 };
diff --git a/Engine/Source/Runtime/MediaAssets/Public/MediaComponent.h b/Engine/Source/Runtime/MediaAssets/Public/MediaComponent.h
--- a/Engine/Source/Runtime/MediaAssets/Public/MediaComponent.h
+++ b/Engine/Source/Runtime/MediaAssets/Public/MediaComponent.h
@@ -9,25 +9,25 @@
 class UMediaPlayer;
 class UMediaTexture;
 
-UCLASS(Category="Media", hideCategories=(Media, Activation))
-class MEDIAASSETS_API UMediaComponent : public UActorComponent
+UCLASS(Category="Media", hideCategories=(Media, Activation), MinimalAPI)
+class UMediaComponent : public UActorComponent
 {
 public:
 	GENERATED_BODY()
 
-	UMediaComponent(const FObjectInitializer& ObjectInitializer);
+	MEDIAASSETS_API UMediaComponent(const FObjectInitializer& ObjectInitializer);
 
 	/** Get the media player that this component owns */
 	UFUNCTION(BlueprintPure, Category="Media")
-	UMediaPlayer* GetMediaPlayer() const;
+	MEDIAASSETS_API UMediaPlayer* GetMediaPlayer() const;
 
 	/** Get the media texture that this component owns, bound to its media player. */
 	UFUNCTION(BlueprintPure, Category="Media")
-	UMediaTexture* GetMediaTexture() const;
+	MEDIAASSETS_API UMediaTexture* GetMediaTexture() const;
 
 private:
 
-	virtual void OnRegister() override;
+	MEDIAASSETS_API virtual void OnRegister() override;
 
 private:
 
@@ -38,4 +38,4 @@ private:
 	/** This component's media player */
 	UPROPERTY(Category="Media", Instanced, Transient, BlueprintGetter=GetMediaPlayer, Interp)
 	TObjectPtr<UMediaPlayer> MediaPlayer;
-};
\ No newline at end of file
+};
