PR #12601: Adding EBufferUsageFlags that was accidentally removed in a refactor.

diff --git a/Engine/Source/Runtime/MRMesh/Private/MRMeshComponent.cpp b/Engine/Source/Runtime/MRMesh/Private/MRMeshComponent.cpp
--- a/Engine/Source/Runtime/MRMesh/Private/MRMeshComponent.cpp
+++ b/Engine/Source/Runtime/MRMesh/Private/MRMeshComponent.cpp
@@ -61,7 +61,7 @@ public:
 	void InitRHIWith(FRHICommandListBase& RHICmdList, const TArray<DataType>& PerVertexData )
 	{
 		NumVerts = PerVertexData.Num();
-		VertexBufferRHI = UE::RHIResourceUtils::CreateVertexBufferFromArray(RHICmdList, TEXT("FMRMeshVertexBuffer"), EBufferUsageFlags::Static, MakeConstArrayView(PerVertexData));
+		VertexBufferRHI = UE::RHIResourceUtils::CreateVertexBufferFromArray(RHICmdList, TEXT("FMRMeshVertexBuffer"), EBufferUsageFlags::Static | EBufferUsageFlags::ShaderResource, MakeConstArrayView(PerVertexData));
 	}
 };
 
